Media: Chnaged WmfMedia plug-in loading errors to warnings

diff --git a/Engine/Plugins/Media/WmfMedia/Source/WmfMedia/Private/WmfMediaModule.cpp b/Engine/Plugins/Media/WmfMedia/Source/WmfMedia/Private/WmfMediaModule.cpp
--- a/Engine/Plugins/Media/WmfMedia/Source/WmfMedia/Private/WmfMediaModule.cpp
+++ b/Engine/Plugins/Media/WmfMedia/Source/WmfMedia/Private/WmfMediaModule.cpp
@@ -41,14 +41,14 @@ public:
 
 		if (MediaModule == nullptr)
 		{
-			UE_LOG(LogWmfMedia, Error, TEXT("Failed to load Media module"));
+			UE_LOG(LogWmfMedia, Warning, TEXT("Failed to load Media module"));
 
 			return;
 		}
 
 		if (!LoadRequiredLibraries())
 		{
-			UE_LOG(LogWmfMedia, Error, TEXT("Failed to load required Windows Media Foundation libraries"));
+			UE_LOG(LogWmfMedia, Warning, TEXT("Failed to load required Windows Media Foundation libraries"));
 
 			return;
 		}
@@ -58,7 +58,7 @@ public:
 
 		if (FAILED(Result))
 		{
-			UE_LOG(LogWmfMedia, Error, TEXT("Failed to initialize Windows Media Foundation, Error %i"), Result);
+			UE_LOG(LogWmfMedia, Warning, TEXT("Failed to initialize Windows Media Foundation, Error %i"), Result);
 
 			return;
 		}
@@ -136,28 +136,28 @@ protected:
 	{
 		if (LoadLibraryW(TEXT("shlwapi.dll")) == nullptr)
 		{
-			UE_LOG(LogWmfMedia, Error, TEXT("Failed to load shlwapi.dll"));
+			UE_LOG(LogWmfMedia, Warning, TEXT("Failed to load shlwapi.dll"));
 
 			return false;
 		}
 
 		if (LoadLibraryW(TEXT("mf.dll")) == nullptr)
 		{
-			UE_LOG(LogWmfMedia, Error, TEXT("Failed to load mf.dll"));
+			UE_LOG(LogWmfMedia, Warning, TEXT("Failed to load mf.dll"));
 
 			return false;
 		}
 
 		if (LoadLibraryW(TEXT("mfplat.dll")) == nullptr)
 		{
-			UE_LOG(LogWmfMedia, Error, TEXT("Failed to load mfplat.dll"));
+			UE_LOG(LogWmfMedia, Warning, TEXT("Failed to load mfplat.dll"));
 
 			return false;
 		}
 
 		if (LoadLibraryW(TEXT("mfplay.dll")) == nullptr)
 		{
-			UE_LOG(LogWmfMedia, Error, TEXT("Failed to load mfplay.dll"));
+			UE_LOG(LogWmfMedia, Warning, TEXT("Failed to load mfplay.dll"));
 
 			return false;
 		}
