Media: Added display name API to media tracks

diff --git a/Engine/Plugins/Media/WmfMedia/Source/WmfMedia/Private/Tracks/WmfMediaTrack.cpp b/Engine/Plugins/Media/WmfMedia/Source/WmfMedia/Private/Tracks/WmfMediaTrack.cpp
--- a/Engine/Plugins/Media/WmfMedia/Source/WmfMedia/Private/Tracks/WmfMediaTrack.cpp
+++ b/Engine/Plugins/Media/WmfMedia/Source/WmfMedia/Private/Tracks/WmfMediaTrack.cpp
@@ -6,6 +6,9 @@
 #include "AllowWindowsPlatformTypes.h"
 
 
+#define LOCTEXT_NAMESPACE "FWmfMediaTrack"
+
+
 /* FWmfMediaTrack structors
  *****************************************************************************/
 
@@ -56,6 +59,28 @@ bool FWmfMediaTrack::Enable( )
 }
 
 
+FText FWmfMediaTrack::GetDisplayName( ) const
+{
+	FText DisplayName;
+
+	if (Name.IsEmpty())
+	{
+		DisplayName = FText::Format(LOCTEXT("UnnamedTrackFormat", "Unnamed Track {0}"), FText::AsNumber((uint32)StreamIndex));
+	}
+	else
+	{
+		DisplayName = FText::FromString(Name);
+	}
+
+	if (Language.IsEmpty())
+	{
+		return DisplayName;
+	}
+
+	return FText::Format(LOCTEXT("LocalizedTrackFormat", "{0} ({1})"), DisplayName, FText::FromString(Language));
+}
+
+
 uint32 FWmfMediaTrack::GetIndex( ) const
 {
 	return StreamIndex;
@@ -101,4 +126,7 @@ void FWmfMediaTrack::RemoveSink( const IMediaSinkRef& Sink )
 }
 
 
+#undef LOCTEXT_NAMESPACE
+
+
 #include "HideWindowsPlatformTypes.h"
diff --git a/Engine/Plugins/Media/WmfMedia/Source/WmfMedia/Private/Tracks/WmfMediaTrack.h b/Engine/Plugins/Media/WmfMedia/Source/WmfMedia/Private/Tracks/WmfMediaTrack.h
--- a/Engine/Plugins/Media/WmfMedia/Source/WmfMedia/Private/Tracks/WmfMediaTrack.h
+++ b/Engine/Plugins/Media/WmfMedia/Source/WmfMedia/Private/Tracks/WmfMediaTrack.h
@@ -34,6 +34,7 @@ public:
 	virtual void AddSink( const IMediaSinkRef& Sink ) override;
 	virtual bool Disable( ) override;
 	virtual bool Enable( ) override;
+	virtual FText GetDisplayName( ) const override;
 	virtual uint32 GetIndex( ) const override;
 	virtual FString GetLanguage( ) const override;
 	virtual FString GetName( ) const override;
@@ -44,25 +45,25 @@ public:
 
 private:
 
-	// The stream's language tag.
+	/** The stream's language tag. */
 	FString Language;
 
-	// The stream's name.
+	/** The stream's name. */
 	FString Name;
 
-	// Pointer to the presentation descriptor.
+	/** Pointer to the presentation descriptor. */
 	TComPtr<IMFPresentationDescriptor> PresentationDescriptor;
 
-	// Whether the stream contains protected content.
+	/** Whether the stream contains protected content. */
 	bool Protected;
 
-	// The sample grabber callback object.
+	/** The sample grabber callback object. */
 	TComPtr<FWmfMediaSampler> Sampler;
 
-	// Pointer to the stream's descriptor.
+	/** Pointer to the stream's descriptor. */
 	TComPtr<IMFStreamDescriptor> StreamDescriptor;
 
-	// The stream's index number in the presentation.
+	/** The stream's index number in the presentation. */
 	DWORD StreamIndex;
 };
 
diff --git a/Engine/Source/Runtime/Media/Public/Interfaces/IMediaTrack.h b/Engine/Source/Runtime/Media/Public/Interfaces/IMediaTrack.h
--- a/Engine/Source/Runtime/Media/Public/Interfaces/IMediaTrack.h
+++ b/Engine/Source/Runtime/Media/Public/Interfaces/IMediaTrack.h
@@ -106,6 +106,17 @@ public:
 	 */
 	virtual const IMediaTrackCaptionDetails& GetCaptionDetails( ) const = 0;
 
+	/**
+	 * Generates the track's display name.
+	 *
+	 * Unlike GetName(), which may return an empty string for unnamed tracks, this method always
+	 * returns a human readable display text for the track's name that will be generated, if needed.
+	 *
+	 * @return Display name text.
+	 * @see GetName
+	 */
+	virtual FText GetDisplayName( ) const = 0;
+
 	/**
 	 * Gets the track's index number.
 	 *
@@ -124,6 +135,7 @@ public:
 	 * Gets the name of the track.
 	 *
 	 * @return Track name string.
+	 * @see GetDisplayName
 	 */
 	virtual FString GetName( ) const = 0;
 
