Media: Implemented proper pausing and looping for WMF player

diff --git a/Engine/Plugins/Media/WmfMedia/Source/WmfMedia/Private/Player/WmfMediaPlayer.cpp b/Engine/Plugins/Media/WmfMedia/Source/WmfMedia/Private/Player/WmfMediaPlayer.cpp
--- a/Engine/Plugins/Media/WmfMedia/Source/WmfMedia/Private/Player/WmfMediaPlayer.cpp
+++ b/Engine/Plugins/Media/WmfMedia/Source/WmfMedia/Private/Player/WmfMediaPlayer.cpp
@@ -452,14 +452,14 @@ bool FWmfMediaPlayer::InitializeMediaSession( IUnknown* SourceObject, const FStr
 
 	UE_LOG(LogWmfMedia, Verbose, TEXT("Added a total of %i tracks"), Tracks.Num());
 
-	// create session
-	MediaSession = new FWmfMediaSession(Topology);
-	MediaUrl = SourceUrl;
-
 	UINT64 PresentationDuration = 0;
 	PresentationDescriptor->GetUINT64(MF_PD_DURATION, &PresentationDuration);
 	Duration = FTimespan(PresentationDuration);
 
+	// create session
+	MediaSession = new FWmfMediaSession(Duration, Topology);
+	MediaUrl = SourceUrl;
+
 	OpenedEvent.Broadcast(SourceUrl);
 
 	return (MediaSession->GetState() != EMediaStates::Error);
diff --git a/Engine/Plugins/Media/WmfMedia/Source/WmfMedia/Private/Wmf/WmfMediaSession.cpp b/Engine/Plugins/Media/WmfMedia/Source/WmfMedia/Private/Wmf/WmfMediaSession.cpp
--- a/Engine/Plugins/Media/WmfMedia/Source/WmfMedia/Private/Wmf/WmfMediaSession.cpp
+++ b/Engine/Plugins/Media/WmfMedia/Source/WmfMedia/Private/Wmf/WmfMediaSession.cpp
@@ -7,10 +7,11 @@
 /* FWmfMediaSession structors
  *****************************************************************************/
 
-FWmfMediaSession::FWmfMediaSession( const TComPtr<IMFTopology>& InTopology )
+FWmfMediaSession::FWmfMediaSession( const FTimespan& InDuration, const TComPtr<IMFTopology>& InTopology )
 	: Capabilities(0)
 	, ChangeRequested(false)
 	, CurrentRate(0.0f)
+	, Duration(InDuration)
 	, LastError(S_OK)
 	, Looping(false)
 	, RefCount(1)
@@ -140,7 +141,11 @@ bool FWmfMediaSession::SetRate( float Rate )
 
 	RequestedRate = Rate;
 
-	if (!FMath::IsNearlyZero(Rate))
+	if (FMath::IsNearlyZero(Rate))
+	{
+		RequestedState = EMediaStates::Paused;
+	}
+	else
 	{
 		RequestedState = EMediaStates::Playing;
 	}
@@ -221,6 +226,11 @@ STDMETHODIMP FWmfMediaSession::Invoke( IMFAsyncResult* AsyncResult )
 				}
 				else if (EventType == MESessionEnded)
 				{
+					if (Looping && (CurrentRate < 0.0f))
+					{
+						RequestedPosition = Duration;
+					}
+
 					UpdateState(EMediaStates::Stopped);
 				}
 				else if (EventType == MESessionPaused)
@@ -377,6 +387,7 @@ bool FWmfMediaSession::ChangeState( )
 			{
 				StateChangePending = true;
 				StateChangeFailed = FAILED(MediaSession->Pause());
+				RequestedPosition = GetPosition();
 			}
 		}
 		else if (RequestedState == EMediaStates::Playing)
@@ -389,7 +400,6 @@ bool FWmfMediaSession::ChangeState( )
 				{
 					StartPosition.vt = VT_I8;
 					StartPosition.hVal.QuadPart = RequestedPosition.GetTicks();
-					RequestedPosition = FTimespan::MinValue();
 				}
 				else
 				{
@@ -439,6 +449,11 @@ void FWmfMediaSession::UpdateState( EMediaStates CompletedState )
 {
 	CurrentState = CompletedState;
 
+	if (CompletedState == EMediaStates::Playing)
+	{
+		RequestedPosition = FTimespan::MinValue();
+	}
+
 	if (CompletedState == RequestedState)
 	{
 		StateChangePending = false;
diff --git a/Engine/Plugins/Media/WmfMedia/Source/WmfMedia/Private/Wmf/WmfMediaSession.h b/Engine/Plugins/Media/WmfMedia/Source/WmfMedia/Private/Wmf/WmfMediaSession.h
--- a/Engine/Plugins/Media/WmfMedia/Source/WmfMedia/Private/Wmf/WmfMediaSession.h
+++ b/Engine/Plugins/Media/WmfMedia/Source/WmfMedia/Private/Wmf/WmfMediaSession.h
@@ -47,9 +47,10 @@ public:
 	/**
 	 * Creates and initializes a new instance.
 	 *
+	 * @param InDuration The duration of the media.
 	 * @param InTopology The playback topology to use.
 	 */
-	FWmfMediaSession( const TComPtr<IMFTopology>& InTopology );
+	FWmfMediaSession( const FTimespan& InDuration, const TComPtr<IMFTopology>& InTopology );
 
 	/** Virtual destructor .*/
 	virtual ~FWmfMediaSession( ) { }
@@ -228,6 +229,9 @@ private:
 	// The current playback state.
 	EMediaStates CurrentState;
 
+	// The duration of the media.
+	FTimespan Duration;
+
 	// The media session that handles all playback.
 	TComPtr<IMFMediaSession> MediaSession;
 
