WmfMedia: Include what you use

diff --git a/Engine/Plugins/Media/WmfMedia/Source/WmfMedia/Private/Tracks/WmfMediaVideoTrack.h b/Engine/Plugins/Media/WmfMedia/Source/WmfMedia/Private/Tracks/WmfMediaVideoTrack.h
--- a/Engine/Plugins/Media/WmfMedia/Source/WmfMedia/Private/Tracks/WmfMediaVideoTrack.h
+++ b/Engine/Plugins/Media/WmfMedia/Source/WmfMedia/Private/Tracks/WmfMediaVideoTrack.h
@@ -5,6 +5,9 @@
 #include "AllowWindowsPlatformTypes.h"
 
 
+/**
+ * Implements a video track using the Windows Media Foundation.
+ */
 class FWmfMediaVideoTrack
 	: public FWmfMediaTrack
 	, public IMediaTrackVideoDetails
diff --git a/Engine/Plugins/Media/WmfMedia/Source/WmfMedia/Private/Wmf/WmfMediaSession.h b/Engine/Plugins/Media/WmfMedia/Source/WmfMedia/Private/Wmf/WmfMediaSession.h
--- a/Engine/Plugins/Media/WmfMedia/Source/WmfMedia/Private/Wmf/WmfMediaSession.h
+++ b/Engine/Plugins/Media/WmfMedia/Source/WmfMedia/Private/Wmf/WmfMediaSession.h
@@ -5,6 +5,10 @@
 #include "AllowWindowsPlatformTypes.h"
 
 
+// forward declarations
+enum class EMediaPlaybackDirections;
+
+
 /**
  * Enumerates possible states of media playback.
  */
diff --git a/Engine/Plugins/Media/WmfMedia/Source/WmfMedia/Private/WmfMediaModule.cpp b/Engine/Plugins/Media/WmfMedia/Source/WmfMedia/Private/WmfMediaModule.cpp
--- a/Engine/Plugins/Media/WmfMedia/Source/WmfMedia/Private/WmfMediaModule.cpp
+++ b/Engine/Plugins/Media/WmfMedia/Source/WmfMedia/Private/WmfMediaModule.cpp
@@ -1,6 +1,10 @@
 // Copyright 1998-2014 Epic Games, Inc. All Rights Reserved.
 
 #include "WmfMediaPrivatePCH.h"
+#include "IMediaModule.h"
+#include "IMediaPlayerFactory.h"
+#include "ModuleInterface.h"
+#include "ModuleManager.h"
 
 #pragma comment(lib, "mf")
 #pragma comment(lib, "mfplat")
diff --git a/Engine/Plugins/Media/WmfMedia/Source/WmfMedia/Private/WmfMediaPrivatePCH.h b/Engine/Plugins/Media/WmfMedia/Source/WmfMedia/Private/WmfMediaPrivatePCH.h
--- a/Engine/Plugins/Media/WmfMedia/Source/WmfMedia/Private/WmfMediaPrivatePCH.h
+++ b/Engine/Plugins/Media/WmfMedia/Source/WmfMedia/Private/WmfMediaPrivatePCH.h
@@ -6,8 +6,16 @@
 /* Dependencies
  *****************************************************************************/
 
-#include "ModuleManager.h"
-#include "Media.h"
+#include "Core.h"
+#include "CoreUObject.h"
+
+#include "IMediaInfo.h"
+#include "IMediaPlayer.h"
+#include "IMediaSink.h"
+#include "IMediaTrack.h"
+#include "IMediaTrackAudioDetails.h"
+#include "IMediaTrackCaptionDetails.h"
+#include "IMediaTrackVideoDetails.h"
 
 #include "AllowWindowsPlatformTypes.h"
 #include <windows.h>
