WmfMedia: Fixed a memory leak due to not properly shutting down media sources & sessions

diff --git a/Engine/Plugins/Media/WmfMedia/Source/WmfMedia/Private/Player/WmfMediaPlayer.cpp b/Engine/Plugins/Media/WmfMedia/Source/WmfMedia/Private/Player/WmfMediaPlayer.cpp
--- a/Engine/Plugins/Media/WmfMedia/Source/WmfMedia/Private/Player/WmfMediaPlayer.cpp
+++ b/Engine/Plugins/Media/WmfMedia/Source/WmfMedia/Private/Player/WmfMediaPlayer.cpp
@@ -80,6 +80,12 @@ void FWmfMediaPlayer::Close()
 	MediaSession->SetState(EMediaStates::Closed);
 	MediaSession.Reset();
 
+	if (MediaSource != NULL)
+	{
+		MediaSource->Shutdown();
+		MediaSource = NULL;
+	}
+
 	Tracks.Reset();
 
 	Duration = 0;
@@ -407,9 +413,9 @@ bool FWmfMediaPlayer::InitializeMediaSession( IUnknown* SourceObject, const FStr
 	UE_LOG(LogWmfMedia, Verbose, TEXT("Initializing media session for %s"), *SourceUrl);
 
 	// create presentation descriptor
-	TComPtr<IMFMediaSource> MediaSource;
+	TComPtr<IMFMediaSource> MediaSourceObject;
 
-	if (FAILED(SourceObject->QueryInterface(IID_PPV_ARGS(&MediaSource))))
+	if (FAILED(SourceObject->QueryInterface(IID_PPV_ARGS(&MediaSourceObject))))
 	{
 		UE_LOG(LogWmfMedia, Error, TEXT("Failed to query media source"));
 
@@ -418,7 +424,7 @@ bool FWmfMediaPlayer::InitializeMediaSession( IUnknown* SourceObject, const FStr
 
 	TComPtr<IMFPresentationDescriptor> PresentationDescriptor;
 	
-	if (FAILED(MediaSource->CreatePresentationDescriptor(&PresentationDescriptor)))
+	if (FAILED(MediaSourceObject->CreatePresentationDescriptor(&PresentationDescriptor)))
 	{
 		UE_LOG(LogWmfMedia, Error, TEXT("Failed to create presentation descriptor"));
 
@@ -446,7 +452,7 @@ bool FWmfMediaPlayer::InitializeMediaSession( IUnknown* SourceObject, const FStr
 
 	for (uint32 StreamIndex = 0; StreamIndex < StreamCount; ++StreamIndex)
 	{
-		AddStreamToTopology(StreamIndex, Topology, PresentationDescriptor, MediaSource);
+		AddStreamToTopology(StreamIndex, Topology, PresentationDescriptor, MediaSourceObject);
 	}
 
 	UE_LOG(LogWmfMedia, Verbose, TEXT("Added a total of %i tracks"), Tracks.Num());
@@ -456,6 +462,7 @@ bool FWmfMediaPlayer::InitializeMediaSession( IUnknown* SourceObject, const FStr
 	Duration = FTimespan(PresentationDuration);
 
 	// create session
+	MediaSource = MediaSourceObject;
 	MediaSession = new FWmfMediaSession(Duration, Topology);
 	MediaUrl = SourceUrl;
 
diff --git a/Engine/Plugins/Media/WmfMedia/Source/WmfMedia/Private/Player/WmfMediaPlayer.h b/Engine/Plugins/Media/WmfMedia/Source/WmfMedia/Private/Player/WmfMediaPlayer.h
--- a/Engine/Plugins/Media/WmfMedia/Source/WmfMedia/Private/Player/WmfMediaPlayer.h
+++ b/Engine/Plugins/Media/WmfMedia/Source/WmfMedia/Private/Player/WmfMediaPlayer.h
@@ -95,6 +95,9 @@ private:
 	/** Holds the asynchronous callback object for the media stream. */
 	TComPtr<FWmfMediaSession> MediaSession;
 
+	/** Holds a pointer to the media source object. */
+	TComPtr<IMFMediaSource> MediaSource;
+
 	/** The URL of the currently opened media. */
 	FString MediaUrl;
 
diff --git a/Engine/Plugins/Media/WmfMedia/Source/WmfMedia/Private/Wmf/WmfMediaByteStream.cpp b/Engine/Plugins/Media/WmfMedia/Source/WmfMedia/Private/Wmf/WmfMediaByteStream.cpp
--- a/Engine/Plugins/Media/WmfMedia/Source/WmfMedia/Private/Wmf/WmfMediaByteStream.cpp
+++ b/Engine/Plugins/Media/WmfMedia/Source/WmfMedia/Private/Wmf/WmfMediaByteStream.cpp
@@ -11,7 +11,7 @@ FWmfMediaByteStream::FWmfMediaByteStream( const TSharedRef<TArray<uint8>>& InBuf
 	: AsyncReadInProgress(false)
 	, Buffer(InBuffer)
 	, Position(0)
-	, RefCount(1)
+	, RefCount(0)
 { }
 
 
diff --git a/Engine/Plugins/Media/WmfMedia/Source/WmfMedia/Private/Wmf/WmfMediaReadState.h b/Engine/Plugins/Media/WmfMedia/Source/WmfMedia/Private/Wmf/WmfMediaReadState.h
--- a/Engine/Plugins/Media/WmfMedia/Source/WmfMedia/Private/Wmf/WmfMediaReadState.h
+++ b/Engine/Plugins/Media/WmfMedia/Source/WmfMedia/Private/Wmf/WmfMediaReadState.h
@@ -27,7 +27,7 @@ public:
 		: BytesRead(0)
 		, ReadBuffer(InReadBuffer)
 		, ReadBufferSize(InReadBufferSize)
-		, RefCount(1)
+		, RefCount(0)
 	{ }
 
 	/** Virtual destructor. */
diff --git a/Engine/Plugins/Media/WmfMedia/Source/WmfMedia/Private/Wmf/WmfMediaSampler.h b/Engine/Plugins/Media/WmfMedia/Source/WmfMedia/Private/Wmf/WmfMediaSampler.h
--- a/Engine/Plugins/Media/WmfMedia/Source/WmfMedia/Private/Wmf/WmfMediaSampler.h
+++ b/Engine/Plugins/Media/WmfMedia/Source/WmfMedia/Private/Wmf/WmfMediaSampler.h
@@ -18,7 +18,7 @@ public:
 
 	/** Default constructor. */
 	FWmfMediaSampler()
-		: RefCount(1)
+		: RefCount(0)
 	{ }
 
 	/** Virtual destructor. */
diff --git a/Engine/Plugins/Media/WmfMedia/Source/WmfMedia/Private/Wmf/WmfMediaSession.cpp b/Engine/Plugins/Media/WmfMedia/Source/WmfMedia/Private/Wmf/WmfMediaSession.cpp
--- a/Engine/Plugins/Media/WmfMedia/Source/WmfMedia/Private/Wmf/WmfMediaSession.cpp
+++ b/Engine/Plugins/Media/WmfMedia/Source/WmfMedia/Private/Wmf/WmfMediaSession.cpp
@@ -14,7 +14,7 @@ FWmfMediaSession::FWmfMediaSession( const FTimespan& InDuration, const TComPtr<I
 	, Duration(InDuration)
 	, LastError(S_OK)
 	, Looping(false)
-	, RefCount(1)
+	, RefCount(0)
 	, RequestedPosition(FTimespan::MinValue())
 	, RequestedRate(0.0f)
 	, StateChangePending(false)
@@ -274,6 +274,7 @@ STDMETHODIMP FWmfMediaSession::Invoke( IMFAsyncResult* AsyncResult )
 					UpdateState(EMediaStates::Stopped);
 				}
 
+				// request the next event
 				if (FAILED(MediaSession->BeginGetEvent(this, NULL)))
 				{
 					Capabilities = 0;
@@ -305,6 +306,11 @@ STDMETHODIMP_(ULONG) FWmfMediaSession::Release()
 	
 	if (CurrentRefCount == 0)
 	{
+		if (MediaSession != NULL)
+		{
+			MediaSession->Shutdown();
+		}
+
 		delete this;
 	}
 
diff --git a/Engine/Plugins/Media/WmfMedia/Source/WmfMedia/Private/Wmf/WmfMediaSession.h b/Engine/Plugins/Media/WmfMedia/Source/WmfMedia/Private/Wmf/WmfMediaSession.h
--- a/Engine/Plugins/Media/WmfMedia/Source/WmfMedia/Private/Wmf/WmfMediaSession.h
+++ b/Engine/Plugins/Media/WmfMedia/Source/WmfMedia/Private/Wmf/WmfMediaSession.h
@@ -56,9 +56,6 @@ public:
 	 */
 	FWmfMediaSession( const FTimespan& InDuration, const TComPtr<IMFTopology>& InTopology );
 
-	/** Virtual destructor .*/
-	virtual ~FWmfMediaSession() { }
-
 public:
 
 	/**
