Change WmfMedia plugin to be RuntimeNoCommandlet, so it doesn't cause warnings on the build machines.

diff --git a/Engine/Plugins/Media/WmfMedia/WmfMedia.uplugin b/Engine/Plugins/Media/WmfMedia/WmfMedia.uplugin
--- a/Engine/Plugins/Media/WmfMedia/WmfMedia.uplugin
+++ b/Engine/Plugins/Media/WmfMedia/WmfMedia.uplugin
@@ -16,7 +16,7 @@
 	[
 		{
 			"Name" : "WmfMedia",
-			"Type" : "Runtime",
+			"Type" : "RuntimeNoCommandlet",
 			"LoadingPhase" : "PreLoadingScreen",
 			"WhitelistPlatforms" : [ "Win32", "Win64" ]
 		}
