Media: reduced nesting

diff --git a/Engine/Plugins/Media/WmfMedia/Source/WmfMedia/Private/Wmf/WmfMediaSession.cpp b/Engine/Plugins/Media/WmfMedia/Source/WmfMedia/Private/Wmf/WmfMediaSession.cpp
--- a/Engine/Plugins/Media/WmfMedia/Source/WmfMedia/Private/Wmf/WmfMediaSession.cpp
+++ b/Engine/Plugins/Media/WmfMedia/Source/WmfMedia/Private/Wmf/WmfMediaSession.cpp
@@ -194,92 +194,94 @@ STDMETHODIMP FWmfMediaSession::Invoke( IMFAsyncResult* AsyncResult )
 
 	MediaEventType EventType = MEUnknown;
 	
-	if (SUCCEEDED(Event->GetType(&EventType)))
+	if (FAILED(Event->GetType(&EventType)))
 	{
-		HRESULT EventResult;
-		Event->GetStatus(&EventResult);
+		return S_OK;
+	}
+
+	HRESULT EventResult;
+	Event->GetStatus(&EventResult);
 
-		FScopeLock ScopeLock(&CriticalSection);
+	FScopeLock ScopeLock(&CriticalSection);
+	{
+		if (EventType == MESessionClosed)
+		{
+			UpdateState(EMediaStates::Closed);
+			Capabilities = 0;
+		}
+		else if (EventType == MEError)
+		{
+			LastError = EventResult;
+		}
+		else
 		{
-			if (EventType == MESessionClosed)
+			if (EventType == MEEndOfPresentation)
 			{
-				UpdateState(EMediaStates::Closed);
-				Capabilities = 0;
+				if (!Looping)
+				{
+					RequestedState = EMediaStates::Stopped;
+					MediaSession->Stop();
+				}
 			}
-			else if (EventType == MEError)
+			else if (EventType == MESessionCapabilitiesChanged)
 			{
-				LastError = EventResult;
+				Capabilities = ::MFGetAttributeUINT32(Event, MF_EVENT_SESSIONCAPS, Capabilities);
 			}
-			else
+			else if (EventType == MESessionEnded)
 			{
-				if (EventType == MEEndOfPresentation)
-				{
-					if (!Looping)
-					{
-						RequestedState = EMediaStates::Stopped;
-						MediaSession->Stop();
-					}
-				}
-				else if (EventType == MESessionCapabilitiesChanged)
+				if (Looping && (CurrentRate < 0.0f))
 				{
-					Capabilities = ::MFGetAttributeUINT32(Event, MF_EVENT_SESSIONCAPS, Capabilities);
+					RequestedPosition = Duration;
 				}
-				else if (EventType == MESessionEnded)
-				{
-					if (Looping && (CurrentRate < 0.0f))
-					{
-						RequestedPosition = Duration;
-					}
 
-					UpdateState(EMediaStates::Stopped);
-				}
-				else if (EventType == MESessionPaused)
-				{
-					UpdateState(EMediaStates::Paused);
-				}
-				else if (EventType == MESessionRateChanged)
+				UpdateState(EMediaStates::Stopped);
+			}
+			else if (EventType == MESessionPaused)
+			{
+				UpdateState(EMediaStates::Paused);
+			}
+			else if (EventType == MESessionRateChanged)
+			{
+				// recover active playback rate if rate change failed
+				if (FAILED(EventResult) && (CurrentRate == RequestedRate))
 				{
-					// recover active playback rate if rate change failed
-					if (FAILED(EventResult) && (CurrentRate == RequestedRate))
-					{
-						PROPVARIANT Value;
-						PropVariantInit(&Value);
-
-						if (SUCCEEDED(Event->GetValue(&Value)) && (Value.vt == VT_R4))
-						{
-							CurrentRate = Value.fltVal;
-						}
+					PROPVARIANT Value;
+					PropVariantInit(&Value);
 
-						RequestedRate = CurrentRate;
-					}
-				}
-				else if (EventType == MESessionScrubSampleComplete)
-				{
-					if (GetInternalPosition() == RequestedPosition)
-					{
-						RequestedPosition = FTimespan::MinValue();
-					}
-				}
-				else if (EventType == MESessionStarted)
-				{
-					if (GetInternalPosition() == RequestedPosition)
+					if (SUCCEEDED(Event->GetValue(&Value)) && (Value.vt == VT_R4))
 					{
-						RequestedPosition = FTimespan::MinValue();
+						CurrentRate = Value.fltVal;
 					}
 
-					UpdateState(EMediaStates::Playing);
+					RequestedRate = CurrentRate;
 				}
-				else if (EventType == MESessionStopped)
+			}
+			else if (EventType == MESessionScrubSampleComplete)
+			{
+				if (GetInternalPosition() == RequestedPosition)
 				{
-					UpdateState(EMediaStates::Stopped);
+					RequestedPosition = FTimespan::MinValue();
 				}
-
-				// request the next event
-				if (FAILED(MediaSession->BeginGetEvent(this, NULL)))
+			}
+			else if (EventType == MESessionStarted)
+			{
+				if (GetInternalPosition() == RequestedPosition)
 				{
-					Capabilities = 0;
-					CurrentState = EMediaStates::Error;
+					RequestedPosition = FTimespan::MinValue();
 				}
+
+				UpdateState(EMediaStates::Playing);
+			}
+			else if (EventType == MESessionStopped)
+			{
+				UpdateState(EMediaStates::Stopped);
+			}
+
+			// request the next event
+			if (FAILED(MediaSession->BeginGetEvent(this, NULL)))
+			{
+				Capabilities = 0;
+				CurrentState = EMediaStates::Error;
 			}
 		}
 	}
