Media: WMF: Shutting down session on unrecoverable error

diff --git a/Engine/Plugins/Media/WmfMedia/Source/WmfMedia/Private/Player/WmfMediaPlayer.cpp b/Engine/Plugins/Media/WmfMedia/Source/WmfMedia/Private/Player/WmfMediaPlayer.cpp
--- a/Engine/Plugins/Media/WmfMedia/Source/WmfMedia/Private/Player/WmfMediaPlayer.cpp
+++ b/Engine/Plugins/Media/WmfMedia/Source/WmfMedia/Private/Player/WmfMediaPlayer.cpp
@@ -258,7 +258,6 @@ bool FWmfMediaPlayer::SetRate( float Rate )
 }
 
 
-
 /* FWmfMediaPlayer implementation
  *****************************************************************************/
 
@@ -462,9 +461,12 @@ bool FWmfMediaPlayer::InitializeMediaSession( IUnknown* SourceObject, const FStr
 	Duration = FTimespan(PresentationDuration);
 
 	// create session
+	MediaUrl = SourceUrl;
 	MediaSource = MediaSourceObject;
 	MediaSession = new FWmfMediaSession(Duration, Topology);
-	MediaUrl = SourceUrl;
+	{
+		MediaSession->OnError().AddRaw(this, &FWmfMediaPlayer::HandleSessionError);
+	}
 
 	OpenedEvent.Broadcast(SourceUrl);
 
@@ -472,4 +474,13 @@ bool FWmfMediaPlayer::InitializeMediaSession( IUnknown* SourceObject, const FStr
 }
 
 
+/* FWmfMediaPlayer callbacks
+ *****************************************************************************/
+
+void FWmfMediaPlayer::HandleSessionError(HRESULT Error)
+{
+	UE_LOG(LogWmfMedia, Error, TEXT("An unrecoverable error occured in the media session: 0x%X"), Error);
+}
+
+
 #include "HideWindowsPlatformTypes.h"
diff --git a/Engine/Plugins/Media/WmfMedia/Source/WmfMedia/Private/Player/WmfMediaPlayer.h b/Engine/Plugins/Media/WmfMedia/Source/WmfMedia/Private/Player/WmfMediaPlayer.h
--- a/Engine/Plugins/Media/WmfMedia/Source/WmfMedia/Private/Player/WmfMediaPlayer.h
+++ b/Engine/Plugins/Media/WmfMedia/Source/WmfMedia/Private/Player/WmfMediaPlayer.h
@@ -84,6 +84,11 @@ protected:
 	 */
 	bool InitializeMediaSession( IUnknown* SourceObject, const FString& SourceUrl );
 
+private:
+
+	/** Handles session errors. */
+	void HandleSessionError(HRESULT Error);
+
 private:
 
 	/** The available media tracks. */
diff --git a/Engine/Plugins/Media/WmfMedia/Source/WmfMedia/Private/Wmf/WmfMediaSession.cpp b/Engine/Plugins/Media/WmfMedia/Source/WmfMedia/Private/Wmf/WmfMediaSession.cpp
--- a/Engine/Plugins/Media/WmfMedia/Source/WmfMedia/Private/Wmf/WmfMediaSession.cpp
+++ b/Engine/Plugins/Media/WmfMedia/Source/WmfMedia/Private/Wmf/WmfMediaSession.cpp
@@ -12,7 +12,6 @@ FWmfMediaSession::FWmfMediaSession( const FTimespan& InDuration, const TComPtr<I
 	, ChangeRequested(false)
 	, CurrentRate(0.0f)
 	, Duration(InDuration)
-	, LastError(S_OK)
 	, Looping(false)
 	, RefCount(0)
 	, RequestedPosition(FTimespan::MinValue())
@@ -211,7 +210,8 @@ STDMETHODIMP FWmfMediaSession::Invoke( IMFAsyncResult* AsyncResult )
 		}
 		else if (EventType == MEError)
 		{
-			LastError = EventResult;
+			ErrorEvent.Broadcast(EventResult);
+			UpdateState(EMediaStates::Error);
 		}
 		else
 		{
diff --git a/Engine/Plugins/Media/WmfMedia/Source/WmfMedia/Private/Wmf/WmfMediaSession.h b/Engine/Plugins/Media/WmfMedia/Source/WmfMedia/Private/Wmf/WmfMediaSession.h
--- a/Engine/Plugins/Media/WmfMedia/Source/WmfMedia/Private/Wmf/WmfMediaSession.h
+++ b/Engine/Plugins/Media/WmfMedia/Source/WmfMedia/Private/Wmf/WmfMediaSession.h
@@ -68,17 +68,6 @@ public:
 		return Capabilities;
 	}
 
-	/**
-	 * Gets the result code of the last error.
-	 *
-	 * @return Error result code.
-	 * @see GetSessionState
-	 */
-	HRESULT GetLastError() const
-	{
-		return LastError;
-	}
-
 	/**
 	 * Gets the current playback position.
 	 *
@@ -100,7 +89,6 @@ public:
 	 * Gets the state of this session.
 	 *
 	 * @return Media session state.
-	 * @see GetLastError
 	 */
 	EMediaStates GetState() const
 	{
@@ -179,6 +167,15 @@ public:
 		return CanScrub;
 	}
 
+public:
+
+	/** Gets an event delegate that is invoked when an error occured. */
+	DECLARE_EVENT_OneParam(FWmfMediaSession, FOnError, HRESULT /*Error*/)
+	FOnError& OnError()
+	{
+		return ErrorEvent;
+	}
+
 public:
 
 	// IMFAsyncCallback interface
@@ -238,9 +235,6 @@ private:
 	/** The media session that handles all playback. */
 	TComPtr<IMFMediaSession> MediaSession;
 
-	/** Caches the last error result code (0 = no error). */
-	HRESULT LastError;
-
 	/** Whether playback should loop to the beginning. */
 	bool Looping;
 
@@ -267,6 +261,11 @@ private:
 
 	/** Whether a state change is currently in progress. */
 	bool StateChangePending;
+
+private:
+
+	/** Holds an event delegate that is invoked when an error occured. */
+	FOnError ErrorEvent;
 };
 
 
