Media: Fixed play button unavailable after non-looped movie finishes

diff --git a/Engine/Plugins/Media/WmfMedia/Source/WmfMedia/Private/Wmf/WmfMediaSession.cpp b/Engine/Plugins/Media/WmfMedia/Source/WmfMedia/Private/Wmf/WmfMediaSession.cpp
--- a/Engine/Plugins/Media/WmfMedia/Source/WmfMedia/Private/Wmf/WmfMediaSession.cpp
+++ b/Engine/Plugins/Media/WmfMedia/Source/WmfMedia/Private/Wmf/WmfMediaSession.cpp
@@ -341,7 +341,6 @@ bool FWmfMediaSession::ChangeState()
 			return true;
 		}
 
-
 		// close session
 		if ((RequestedState == EMediaStates::Closed) || (RequestedState == EMediaStates::Error))
 		{
@@ -465,6 +464,12 @@ void FWmfMediaSession::UpdateState( EMediaStates CompletedState )
 	if (CompletedState == RequestedState)
 	{
 		StateChangePending = false;
+
+		if (RequestedState == EMediaStates::Stopped)
+		{
+			CurrentRate = 0.0f;
+			RequestedRate = 0.0f;
+		}
 	}
 	else
 	{
