Media: Added support for extension-less streaming media URI schemes by removing assumptions about media URL formats

#UpgradeNotes
- IMediaModule::GetSupportedFormats() has been renamed to IMediaModule::GetSupportedFileTypes()
- FMediaFormats has been renamed to FMediaFileTypes
- IMediaPlayerFactory::GetSupportedFormats() has been renamed to IMediaPlayerFactory::GetSupportedFileTypes()
- IMediaPlayerFactory::SupportsFile() has been replaced with IMediaPlayerFactory::SupportsUrl(), which must now be implemented by media plug-ins

diff --git a/Engine/Plugins/Media/AndroidMedia/Source/AndroidMedia/Private/AndroidMediaModule.cpp b/Engine/Plugins/Media/AndroidMedia/Source/AndroidMedia/Private/AndroidMediaModule.cpp
--- a/Engine/Plugins/Media/AndroidMedia/Source/AndroidMedia/Private/AndroidMediaModule.cpp
+++ b/Engine/Plugins/Media/AndroidMedia/Source/AndroidMedia/Private/AndroidMediaModule.cpp
@@ -21,9 +21,9 @@ public:
 			IMediaModule* MediaModule = FModuleManager::LoadModulePtr<IMediaModule>("Media");
 			if (nullptr != MediaModule)
 			{
-				SupportedFormats.Add(TEXT("3gpp"), LOCTEXT("Format3gpp", "3GPP Multimedia File"));
-				SupportedFormats.Add(TEXT("aac"), LOCTEXT("FormatAac", "MPEG-2 Advanced Audio Coding File"));
-				SupportedFormats.Add(TEXT("mp4"), LOCTEXT("FormatMp4", "MPEG-4 Movie"));
+				SupportedFileTypes.Add(TEXT("3gpp"), LOCTEXT("Format3gpp", "3GPP Multimedia File"));
+				SupportedFileTypes.Add(TEXT("aac"), LOCTEXT("FormatAac", "MPEG-2 Advanced Audio Coding File"));
+				SupportedFileTypes.Add(TEXT("mp4"), LOCTEXT("FormatMp4", "MPEG-4 Movie"));
 				MediaModule->RegisterPlayerFactory(*this);
 			}
 		}
@@ -55,15 +55,20 @@ public:
 		}
 	}
 
-	virtual const FMediaFormats& GetSupportedFormats() const override
+	virtual const FMediaFileTypes& GetSupportedFileTypes() const override
 	{
-		return SupportedFormats;
+		return SupportedFileTypes;
+	}
+
+	virtual bool SupportsUrl(const FString& Url) const override
+	{
+		return SupportedFileTypes.Contains(FPaths::GetExtension(Url));
 	}
 
 private:
 
-	// The collection of supported media formats.
-	FMediaFormats SupportedFormats;
+	// The collection of supported media file types.
+	FMediaFileTypes SupportedFileTypes;
 
 	bool IsSupported()
 	{
diff --git a/Engine/Plugins/Media/AvfMedia/Source/AvfMedia/Private/AvfMediaModule.cpp b/Engine/Plugins/Media/AvfMedia/Source/AvfMedia/Private/AvfMediaModule.cpp
--- a/Engine/Plugins/Media/AvfMedia/Source/AvfMedia/Private/AvfMediaModule.cpp
+++ b/Engine/Plugins/Media/AvfMedia/Source/AvfMedia/Private/AvfMediaModule.cpp
@@ -43,10 +43,10 @@ public:
 		}
 
 		// initialize supported media formats
-		SupportedFormats.Add(TEXT("avi"), LOCTEXT("FormatAvi", "Audio Video Interleave File"));
-		SupportedFormats.Add(TEXT("m4v"), LOCTEXT("FormatM4v", "Apple MPEG-4 Video"));
-		SupportedFormats.Add(TEXT("mov"), LOCTEXT("FormatMov", "Apple QuickTime Movie"));
-		SupportedFormats.Add(TEXT("mp4"), LOCTEXT("FormatMp4", "MPEG-4 Movie"));
+		SupportedFileTypes.Add(TEXT("avi"), LOCTEXT("FormatAvi", "Audio Video Interleave File"));
+		SupportedFileTypes.Add(TEXT("m4v"), LOCTEXT("FormatM4v", "Apple MPEG-4 Video"));
+		SupportedFileTypes.Add(TEXT("mov"), LOCTEXT("FormatMov", "Apple QuickTime Movie"));
+		SupportedFileTypes.Add(TEXT("mp4"), LOCTEXT("FormatMp4", "MPEG-4 Movie"));
 
 		// register factory
 		MediaModule->RegisterPlayerFactory(*this);
@@ -87,9 +87,14 @@ public:
 		return nullptr;
 	}
 
-	virtual const FMediaFormats& GetSupportedFormats() const override
+	virtual const FMediaFileTypes& GetSupportedFileTypes() const override
 	{
-		return SupportedFormats;
+		return SupportedFileTypes;
+	}
+
+	virtual bool SupportsUrl(const FString& Url) const override
+	{
+		return SupportedFileTypes.Contains(FPaths::GetExtension(Url));
 	}
 
 protected:
@@ -109,8 +114,8 @@ private:
 	/** Whether the module has been initialized. */
 	bool Initialized;
 
-	/** The collection of supported media formats. */
-	FMediaFormats SupportedFormats;
+	/** The collection of supported media file types. */
+	FMediaFileTypes SupportedFileTypes;
 };
 
 
diff --git a/Engine/Plugins/Media/MediaPlayerEditor/Source/MediaPlayerEditor/Private/Factories/MediaPlayerFactory.cpp b/Engine/Plugins/Media/MediaPlayerEditor/Source/MediaPlayerEditor/Private/Factories/MediaPlayerFactory.cpp
--- a/Engine/Plugins/Media/MediaPlayerEditor/Source/MediaPlayerEditor/Private/Factories/MediaPlayerFactory.cpp
+++ b/Engine/Plugins/Media/MediaPlayerEditor/Source/MediaPlayerEditor/Private/Factories/MediaPlayerFactory.cpp
@@ -52,10 +52,10 @@ void UMediaPlayerFactory::ReloadMediaFormats()
 
 	if (MediaModule != nullptr)
 	{
-		FMediaFormats SupportedFormats;
-		MediaModule->GetSupportedFormats(SupportedFormats);
+		FMediaFileTypes FileTypes;
+		MediaModule->GetSupportedFileTypes(FileTypes);
 
-		for (auto& Format : SupportedFormats)
+		for (auto& Format : FileTypes)
 		{
 			Formats.Add(Format.Key + TEXT(";") + Format.Value.ToString());
 		}
@@ -71,6 +71,7 @@ void UMediaPlayerFactory::HandleMediaPlayerFactoryAdded()
 	ReloadMediaFormats();
 }
 
+
 void UMediaPlayerFactory::HandleMediaPlayerFactoryRemoved()
 {
 	ReloadMediaFormats();
diff --git a/Engine/Plugins/Media/WmfMedia/Source/WmfMedia/Private/WmfMediaModule.cpp b/Engine/Plugins/Media/WmfMedia/Source/WmfMedia/Private/WmfMediaModule.cpp
--- a/Engine/Plugins/Media/WmfMedia/Source/WmfMedia/Private/WmfMediaModule.cpp
+++ b/Engine/Plugins/Media/WmfMedia/Source/WmfMedia/Private/WmfMediaModule.cpp
@@ -66,20 +66,20 @@ public:
 		}
 
 		// initialize supported media formats
-		SupportedFormats.Add(TEXT("3g2"), LOCTEXT("Format3g2", "3G2 Multimedia Stream"));
-		SupportedFormats.Add(TEXT("3gp"), LOCTEXT("Format3gp", "3GP Video Stream"));
-		SupportedFormats.Add(TEXT("3gp2"), LOCTEXT("Format3gp2", "3GPP2 Multimedia File"));
-		SupportedFormats.Add(TEXT("3gpp"), LOCTEXT("Format3gpp", "3GPP Multimedia File"));
-		SupportedFormats.Add(TEXT("aac"), LOCTEXT("FormatAac", "MPEG-2 Advanced Audio Coding File"));
-		SupportedFormats.Add(TEXT("adts"), LOCTEXT("FormatAdts", "Audio Data Transport Stream"));
-		SupportedFormats.Add(TEXT("asf"), LOCTEXT("FormatAsf", "ASF Media File"));
-		SupportedFormats.Add(TEXT("avi"), LOCTEXT("FormatAvi", "Audio Video Interleave File"));
-		SupportedFormats.Add(TEXT("m4v"), LOCTEXT("FormatM4v", "Apple MPEG-4 Video"));
-		SupportedFormats.Add(TEXT("mov"), LOCTEXT("FormatMov", "Apple QuickTime Movie"));
-		SupportedFormats.Add(TEXT("mp4"), LOCTEXT("FormatMp4", "MPEG-4 Movie"));
-		SupportedFormats.Add(TEXT("sami"), LOCTEXT("FormatSami", "Synchronized Accessible Media Interchange (SAMI) File"));
-		SupportedFormats.Add(TEXT("smi"), LOCTEXT("FormatSmi", "Synchronized Multimedia Integration (SMIL) File"));
-		SupportedFormats.Add(TEXT("wmv"), LOCTEXT("FormatWmv", "Windows Media Video"));
+		SupportedFileTypes.Add(TEXT("3g2"), LOCTEXT("Format3g2", "3G2 Multimedia Stream"));
+		SupportedFileTypes.Add(TEXT("3gp"), LOCTEXT("Format3gp", "3GP Video Stream"));
+		SupportedFileTypes.Add(TEXT("3gp2"), LOCTEXT("Format3gp2", "3GPP2 Multimedia File"));
+		SupportedFileTypes.Add(TEXT("3gpp"), LOCTEXT("Format3gpp", "3GPP Multimedia File"));
+		SupportedFileTypes.Add(TEXT("aac"), LOCTEXT("FormatAac", "MPEG-2 Advanced Audio Coding File"));
+		SupportedFileTypes.Add(TEXT("adts"), LOCTEXT("FormatAdts", "Audio Data Transport Stream"));
+		SupportedFileTypes.Add(TEXT("asf"), LOCTEXT("FormatAsf", "ASF Media File"));
+		SupportedFileTypes.Add(TEXT("avi"), LOCTEXT("FormatAvi", "Audio Video Interleave File"));
+		SupportedFileTypes.Add(TEXT("m4v"), LOCTEXT("FormatM4v", "Apple MPEG-4 Video"));
+		SupportedFileTypes.Add(TEXT("mov"), LOCTEXT("FormatMov", "Apple QuickTime Movie"));
+		SupportedFileTypes.Add(TEXT("mp4"), LOCTEXT("FormatMp4", "MPEG-4 Movie"));
+		SupportedFileTypes.Add(TEXT("sami"), LOCTEXT("FormatSami", "Synchronized Accessible Media Interchange (SAMI) File"));
+		SupportedFileTypes.Add(TEXT("smi"), LOCTEXT("FormatSmi", "Synchronized Multimedia Integration (SMIL) File"));
+		SupportedFileTypes.Add(TEXT("wmv"), LOCTEXT("FormatWmv", "Windows Media Video"));
 
 		// register factory
 		MediaModule->RegisterPlayerFactory(*this);
@@ -122,9 +122,14 @@ public:
 		return nullptr;
 	}
 
-	virtual const FMediaFormats& GetSupportedFormats() const override
+	virtual const FMediaFileTypes& GetSupportedFileTypes() const override
 	{
-		return SupportedFormats;
+		return SupportedFileTypes;
+	}
+
+	virtual bool SupportsUrl(const FString& Url) const override
+	{
+		return SupportedFileTypes.Contains(FPaths::GetExtension(Url));
 	}
 
 protected:
@@ -172,8 +177,8 @@ private:
 	/** Whether the module has been initialized. */
 	bool Initialized;
 
-	/** The collection of supported media formats. */
-	FMediaFormats SupportedFormats;
+	/** The collection of supported media file types. */
+	FMediaFileTypes SupportedFileTypes;
 };
 
 
diff --git a/Engine/Source/Editor/DetailCustomizations/Private/MediaPlayerCustomization.cpp b/Engine/Source/Editor/DetailCustomizations/Private/MediaPlayerCustomization.cpp
--- a/Engine/Source/Editor/DetailCustomizations/Private/MediaPlayerCustomization.cpp
+++ b/Engine/Source/Editor/DetailCustomizations/Private/MediaPlayerCustomization.cpp
@@ -358,10 +358,10 @@ FString FMediaPlayerCustomization::HandleUrlPickerFileTypeFilter() const
 		return Filter;
 	}
 
-	FMediaFormats SupportedFormats;
-	MediaModule->GetSupportedFormats(SupportedFormats);
+	FMediaFileTypes FileTypes;
+	MediaModule->GetSupportedFileTypes(FileTypes);
 
-	if (SupportedFormats.Num() == 0)
+	if (FileTypes.Num() == 0)
 	{
 		return Filter;
 	}
@@ -369,7 +369,7 @@ FString FMediaPlayerCustomization::HandleUrlPickerFileTypeFilter() const
 	FString AllExtensions;
 	FString AllFilters;
 			
-	for (auto& Format : SupportedFormats)
+	for (auto& Format : FileTypes)
 	{
 		if (!AllExtensions.IsEmpty())
 		{
diff --git a/Engine/Source/Runtime/Media/Private/MediaModule.cpp b/Engine/Source/Runtime/Media/Private/MediaModule.cpp
--- a/Engine/Source/Runtime/Media/Private/MediaModule.cpp
+++ b/Engine/Source/Runtime/Media/Private/MediaModule.cpp
@@ -39,7 +39,7 @@ public:
 
 		for (IMediaPlayerFactory* PlayerFactory : MediaPlayerFactories)
 		{
-			if (PlayerFactory->SupportsFile(Url))
+			if (PlayerFactory->SupportsUrl(Url))
 			{
 				Player = PlayerFactory->CreatePlayer();
 
@@ -53,16 +53,16 @@ public:
 		return Player;
 	}
 
-	virtual int32 GetSupportedFormats( FMediaFormats& OutFormats ) override
+	virtual int32 GetSupportedFileTypes( FMediaFileTypes& OutFileTypes ) override
 	{
-		OutFormats.Reset();
+		OutFileTypes.Reset();
 
 		for (IMediaPlayerFactory* Factory : MediaPlayerFactories)
 		{
-			OutFormats.Append(Factory->GetSupportedFormats());
+			OutFileTypes.Append(Factory->GetSupportedFileTypes());
 		}
 	
-		return OutFormats.Num();
+		return OutFileTypes.Num();
 	}
 
 	DECLARE_DERIVED_EVENT(FMediaModule, IMediaModule::FOnFactoryAdded, FOnFactoryAdded);
diff --git a/Engine/Source/Runtime/Media/Public/IMediaModule.h b/Engine/Source/Runtime/Media/Public/IMediaModule.h
--- a/Engine/Source/Runtime/Media/Public/IMediaModule.h
+++ b/Engine/Source/Runtime/Media/Public/IMediaModule.h
@@ -24,14 +24,14 @@ public:
 	 *
 	 * @param Url The URL to the media file container to create the player for.
 	 */
-	virtual TSharedPtr<IMediaPlayer> CreatePlayer( const FString& Url ) = 0;
+	virtual TSharedPtr<IMediaPlayer> CreatePlayer(const FString& Url) = 0;
 
 	/**
-	 * Gets the collection of supported media formats.
+	 * Gets the collection of supported media file types.
 	 *
-	 * @param OutFormats Will hold the supported formats.
+	 * @param OutFileTypes Will hold the supported file types.
 	 */
-	virtual int32 GetSupportedFormats( TMap<FString, FText>& OutFormats ) = 0;
+	virtual int32 GetSupportedFileTypes(TMap<FString, FText>& OutFileTypes) = 0;
 
 	/**
 	 * Registers a media player factory.
@@ -39,7 +39,7 @@ public:
 	 * @param Factory The media player factory to register.
 	 * @see UnregisterPlayerFactory
 	 */
-	virtual void RegisterPlayerFactory( IMediaPlayerFactory& Factory ) = 0;
+	virtual void RegisterPlayerFactory(IMediaPlayerFactory& Factory) = 0;
 
 	/**
 	 * Unregisters a media player factory.
diff --git a/Engine/Source/Runtime/Media/Public/IMediaPlayerFactory.h b/Engine/Source/Runtime/Media/Public/IMediaPlayerFactory.h
--- a/Engine/Source/Runtime/Media/Public/IMediaPlayerFactory.h
+++ b/Engine/Source/Runtime/Media/Public/IMediaPlayerFactory.h
@@ -8,22 +8,23 @@ class IMediaPlayer;
 
 
 /**
- * Type definition for media file container formats.
+ * Type definition for media container file types.
  *
  * This type is a TMap that maps the file extensions of supported media file containers
  * formats to a human readable description text. The file extensions are stored without
  * a leading dot, i.e. {"avi", "Audio Video Interleave File"}.
  */
-typedef TMap<FString, FText> FMediaFormats;
+typedef TMap<FString, FText> FMediaFileTypes;
 
 
 /**
  * Interface for media player factories.
  *
- * Video player factories are used to create instances of media player implementations.
+ * Media player factories are used to create instances of media player implementations.
  * Most media players will be implemented inside plug-ins, which will register their
- * factories on startup. By associating a media format with a factory, the Video module
- * can provide available media players for a given format.
+ * factories on startup. The Media module will use the SupportsUrl() method on this
+ * interface to determine which media player to instantiate for a given media URL, such
+ * as a file path or a resource identifier on the internet.
  */
 class IMediaPlayerFactory
 {
@@ -37,25 +38,26 @@ public:
 	virtual TSharedPtr<IMediaPlayer> CreatePlayer() = 0;
 
 	/**
-	 * Gets the collection of supported media formats.
+	 * Gets a read-only collection of supported media file types.
 	 *
-	 * @return A collection of supported formats.
-	 * @see SupportsFile
+	 * This function is used for file based media containers that reside on the users
+	 * local hard drive, a shared network drive or a location on the internet. It does
+	 * not take into account support for extension-less streaming media URI schemes,
+	 * such as rtsp:// and udp://
+	 *
+	 * @return A collection of supported file types.
+	 * @see SupportsUrl
 	 */
-	virtual const FMediaFormats& GetSupportedFormats() const = 0;
-
-public:
+	virtual const FMediaFileTypes& GetSupportedFileTypes() const = 0;
 
 	/**
-	 * Checks whether this factory can create media players that support the specified file format.
+	 * Checks whether this factory can create media players that support the given URL.
 	 *
-	 * @param Url The URL to the media file container.
-	 * @see GetSupportedFormats
+	 * @param Url The URL to the media, i.e. a file path or URI.
+	 * @return true if the URL supported, false otherwise.
+	 * @see GetSupportedFileTypes
 	 */
-	bool SupportsFile( const FString& Url ) const
-	{
-		return GetSupportedFormats().Contains(FPaths::GetExtension(Url));
-	}
+	virtual bool SupportsUrl(const FString& Url) const = 0;
 
 public:
 
