Media: Fixed WmfMedia drawing ClearColor in between looped video playback

diff --git a/Engine/Plugins/Media/WmfMedia/Source/WmfMedia/Private/Wmf/WmfMediaSession.cpp b/Engine/Plugins/Media/WmfMedia/Source/WmfMedia/Private/Wmf/WmfMediaSession.cpp
--- a/Engine/Plugins/Media/WmfMedia/Source/WmfMedia/Private/Wmf/WmfMediaSession.cpp
+++ b/Engine/Plugins/Media/WmfMedia/Source/WmfMedia/Private/Wmf/WmfMediaSession.cpp
@@ -217,7 +217,16 @@ STDMETHODIMP FWmfMediaSession::Invoke( IMFAsyncResult* AsyncResult )
 			{
 				if (EventType == MEEndOfPresentation)
 				{
-					if (!Looping)
+					if (Looping)
+					{
+						if (CurrentRate < 0.0f)
+						{
+							RequestedPosition = Duration;
+						}
+
+						UpdateState(EMediaStates::Paused);
+					}
+					else
 					{
 						RequestedState = EMediaStates::Stopped;
 						MediaSession->Stop();
@@ -229,12 +238,7 @@ STDMETHODIMP FWmfMediaSession::Invoke( IMFAsyncResult* AsyncResult )
 				}
 				else if (EventType == MESessionEnded)
 				{
-					if (Looping && (CurrentRate < 0.0f))
-					{
-						RequestedPosition = Duration;
-					}
-
-					UpdateState(EMediaStates::Stopped);
+					// do nothing
 				}
 				else if (EventType == MESessionPaused)
 				{
