Media: Added supported audio file extensions to WMF plug-in

diff --git a/Engine/Plugins/Media/WmfMedia/Source/WmfMedia/Private/WmfMediaModule.cpp b/Engine/Plugins/Media/WmfMedia/Source/WmfMedia/Private/WmfMediaModule.cpp
--- a/Engine/Plugins/Media/WmfMedia/Source/WmfMedia/Private/WmfMediaModule.cpp
+++ b/Engine/Plugins/Media/WmfMedia/Source/WmfMedia/Private/WmfMediaModule.cpp
@@ -76,9 +76,12 @@ public:
 		SupportedFileTypes.Add(TEXT("avi"), LOCTEXT("FormatAvi", "Audio Video Interleave File"));
 		SupportedFileTypes.Add(TEXT("m4v"), LOCTEXT("FormatM4v", "Apple MPEG-4 Video"));
 		SupportedFileTypes.Add(TEXT("mov"), LOCTEXT("FormatMov", "Apple QuickTime Movie"));
+		SupportedFileTypes.Add(TEXT("mp3"), LOCTEXT("FormatMp3", "MPEG-2 Audio"));
 		SupportedFileTypes.Add(TEXT("mp4"), LOCTEXT("FormatMp4", "MPEG-4 Movie"));
 		SupportedFileTypes.Add(TEXT("sami"), LOCTEXT("FormatSami", "Synchronized Accessible Media Interchange (SAMI) File"));
 		SupportedFileTypes.Add(TEXT("smi"), LOCTEXT("FormatSmi", "Synchronized Multimedia Integration (SMIL) File"));
+		SupportedFileTypes.Add(TEXT("wav"), LOCTEXT("FormatWav", "Wave Audio File"));
+		SupportedFileTypes.Add(TEXT("wma"), LOCTEXT("FormatWma", "Windows Media Audio"));
 		SupportedFileTypes.Add(TEXT("wmv"), LOCTEXT("FormatWmv", "Windows Media Video"));
 
 		// register factory
