Media: Added filter for supported media URI schemes in WMF

diff --git a/Engine/Plugins/Media/WmfMedia/Source/WmfMedia/Private/WmfMediaModule.cpp b/Engine/Plugins/Media/WmfMedia/Source/WmfMedia/Private/WmfMediaModule.cpp
--- a/Engine/Plugins/Media/WmfMedia/Source/WmfMedia/Private/WmfMediaModule.cpp
+++ b/Engine/Plugins/Media/WmfMedia/Source/WmfMedia/Private/WmfMediaModule.cpp
@@ -84,6 +84,15 @@ public:
 		SupportedFileTypes.Add(TEXT("wma"), LOCTEXT("FormatWma", "Windows Media Audio"));
 		SupportedFileTypes.Add(TEXT("wmv"), LOCTEXT("FormatWmv", "Windows Media Video"));
 
+		// initialize supported URI schemes
+		SupportedUriSchemes.Add(TEXT("http://"));
+		SupportedUriSchemes.Add(TEXT("httpd://"));
+		SupportedUriSchemes.Add(TEXT("https://"));
+		SupportedUriSchemes.Add(TEXT("mms://"));
+		SupportedUriSchemes.Add(TEXT("rtsp://"));
+		SupportedUriSchemes.Add(TEXT("rtspt://"));
+		SupportedUriSchemes.Add(TEXT("rtspu://"));
+
 		// register factory
 		MediaModule->RegisterPlayerFactory(*this);
 
@@ -132,7 +141,22 @@ public:
 
 	virtual bool SupportsUrl(const FString& Url) const override
 	{
-		return SupportedFileTypes.Contains(FPaths::GetExtension(Url));
+		const FString Extension = FPaths::GetExtension(Url);
+
+		if (!Extension.IsEmpty())
+		{
+			return SupportedFileTypes.Contains(Extension);
+		}
+
+		for (const FString& Scheme : SupportedUriSchemes)
+		{
+			if (Url.StartsWith(Scheme))
+			{
+				return true;
+			}
+		}
+
+		return false;
 	}
 
 protected:
@@ -182,6 +206,9 @@ private:
 
 	/** The collection of supported media file types. */
 	FMediaFileTypes SupportedFileTypes;
+
+	/** The collection of supported URI schemes. */
+	TArray<FString> SupportedUriSchemes;
 };
 
 
