Media: Buffer parameter in IMediaPlayer::Open is now thread-safe

#UpgradeNotes: If you implement the IMediaPlayer interface, make sure to update the function signature of the Open() method.

diff --git a/Engine/Plugins/Media/AndroidMedia/Source/AndroidMedia/Private/AndroidMediaPlayer.cpp b/Engine/Plugins/Media/AndroidMedia/Source/AndroidMedia/Private/AndroidMediaPlayer.cpp
--- a/Engine/Plugins/Media/AndroidMedia/Source/AndroidMedia/Private/AndroidMediaPlayer.cpp
+++ b/Engine/Plugins/Media/AndroidMedia/Source/AndroidMedia/Private/AndroidMediaPlayer.cpp
@@ -533,7 +533,7 @@ bool FAndroidMediaPlayer::Open(const FString& Url)
 	return MediaState == EMediaState::Prepared;
 }
 
-bool FAndroidMediaPlayer::Open(const TSharedRef<TArray<uint8>>& Buffer,
+bool FAndroidMediaPlayer::Open(const TSharedRef<TArray<uint8>, ESPMode::ThreadSafe>& Buffer,
 	const FString& OriginalUrl)
 {
 	return false;
diff --git a/Engine/Plugins/Media/AndroidMedia/Source/AndroidMedia/Private/AndroidMediaPlayer.h b/Engine/Plugins/Media/AndroidMedia/Source/AndroidMedia/Private/AndroidMediaPlayer.h
--- a/Engine/Plugins/Media/AndroidMedia/Source/AndroidMedia/Private/AndroidMediaPlayer.h
+++ b/Engine/Plugins/Media/AndroidMedia/Source/AndroidMedia/Private/AndroidMediaPlayer.h
@@ -41,7 +41,7 @@ public:
 	virtual bool IsPlaying() const override;
 	virtual bool IsReady() const override;
 	virtual bool Open(const FString& Url) override;
-	virtual bool Open(const TSharedRef<TArray<uint8>>& Buffer,
+	virtual bool Open(const TSharedRef<TArray<uint8>, ESPMode::ThreadSafe>& Buffer,
 		const FString& OriginalUrl) override;
 	virtual bool Seek(const FTimespan& Time) override;
 	virtual bool SetLooping(bool Looping) override;
diff --git a/Engine/Plugins/Media/AvfMedia/Source/AvfMedia/Private/Player/AvfMediaPlayer.cpp b/Engine/Plugins/Media/AvfMedia/Source/AvfMedia/Private/Player/AvfMediaPlayer.cpp
--- a/Engine/Plugins/Media/AvfMedia/Source/AvfMedia/Private/Player/AvfMediaPlayer.cpp
+++ b/Engine/Plugins/Media/AvfMedia/Source/AvfMedia/Private/Player/AvfMediaPlayer.cpp
@@ -318,7 +318,7 @@ bool FAvfMediaPlayer::Open( const FString& Url )
 }
 
 
-bool FAvfMediaPlayer::Open( const TSharedRef<TArray<uint8>>& Buffer, const FString& OriginalUrl )
+bool FAvfMediaPlayer::Open( const TSharedRef<TArray<uint8>, ESPMode::ThreadSafe>& Buffer, const FString& OriginalUrl )
 {
     return false;
 }
diff --git a/Engine/Plugins/Media/AvfMedia/Source/AvfMedia/Private/Player/AvfMediaPlayer.h b/Engine/Plugins/Media/AvfMedia/Source/AvfMedia/Private/Player/AvfMediaPlayer.h
--- a/Engine/Plugins/Media/AvfMedia/Source/AvfMedia/Private/Player/AvfMediaPlayer.h
+++ b/Engine/Plugins/Media/AvfMedia/Source/AvfMedia/Private/Player/AvfMediaPlayer.h
@@ -67,7 +67,7 @@ public:
 	virtual bool IsPlaying() const override;
 	virtual bool IsReady() const override;
 	virtual bool Open( const FString& Url ) override;
-	virtual bool Open( const TSharedRef<TArray<uint8>>& Buffer, const FString& OriginalUrl ) override;
+	virtual bool Open( const TSharedRef<TArray<uint8>, ESPMode::ThreadSafe>& Buffer, const FString& OriginalUrl ) override;
 	virtual bool Seek( const FTimespan& Time ) override;
 	virtual bool SetLooping( bool Looping ) override;
 	virtual bool SetRate( float Rate ) override;
diff --git a/Engine/Plugins/Media/WmfMedia/Source/WmfMedia/Private/Player/WmfMediaPlayer.cpp b/Engine/Plugins/Media/WmfMedia/Source/WmfMedia/Private/Player/WmfMediaPlayer.cpp
--- a/Engine/Plugins/Media/WmfMedia/Source/WmfMedia/Private/Player/WmfMediaPlayer.cpp
+++ b/Engine/Plugins/Media/WmfMedia/Source/WmfMedia/Private/Player/WmfMediaPlayer.cpp
@@ -188,7 +188,7 @@ bool FWmfMediaPlayer::Open( const FString& Url )
 }
 
 
-bool FWmfMediaPlayer::Open( const TSharedRef<TArray<uint8>>& Buffer, const FString& OriginalUrl )
+bool FWmfMediaPlayer::Open( const TSharedRef<TArray<uint8>, ESPMode::ThreadSafe>& Buffer, const FString& OriginalUrl )
 {
 	if ((Buffer->Num() == 0) || OriginalUrl.IsEmpty())
 	{
diff --git a/Engine/Plugins/Media/WmfMedia/Source/WmfMedia/Private/Player/WmfMediaPlayer.h b/Engine/Plugins/Media/WmfMedia/Source/WmfMedia/Private/Player/WmfMediaPlayer.h
--- a/Engine/Plugins/Media/WmfMedia/Source/WmfMedia/Private/Player/WmfMediaPlayer.h
+++ b/Engine/Plugins/Media/WmfMedia/Source/WmfMedia/Private/Player/WmfMediaPlayer.h
@@ -45,7 +45,7 @@ public:
 	virtual bool IsPlaying() const override;
 	virtual bool IsReady() const override;
 	virtual bool Open( const FString& Url ) override;
-	virtual bool Open( const TSharedRef<TArray<uint8>>& Buffer, const FString& OriginalUrl ) override;
+	virtual bool Open( const TSharedRef<TArray<uint8>, ESPMode::ThreadSafe>& Buffer, const FString& OriginalUrl ) override;
 	virtual bool Seek( const FTimespan& Time ) override;
 	virtual bool SetLooping( bool Looping ) override;
 	virtual bool SetRate( float Rate ) override;
diff --git a/Engine/Plugins/Media/WmfMedia/Source/WmfMedia/Private/Wmf/WmfMediaByteStream.cpp b/Engine/Plugins/Media/WmfMedia/Source/WmfMedia/Private/Wmf/WmfMediaByteStream.cpp
--- a/Engine/Plugins/Media/WmfMedia/Source/WmfMedia/Private/Wmf/WmfMediaByteStream.cpp
+++ b/Engine/Plugins/Media/WmfMedia/Source/WmfMedia/Private/Wmf/WmfMediaByteStream.cpp
@@ -7,7 +7,7 @@
 /* FWmfByteStream structors
  *****************************************************************************/
 
-FWmfMediaByteStream::FWmfMediaByteStream( const TSharedRef<TArray<uint8>>& InBuffer )
+FWmfMediaByteStream::FWmfMediaByteStream( const TSharedRef<TArray<uint8>, ESPMode::ThreadSafe>& InBuffer )
 	: AsyncReadInProgress(false)
 	, Buffer(InBuffer)
 	, Position(0)
diff --git a/Engine/Plugins/Media/WmfMedia/Source/WmfMedia/Private/Wmf/WmfMediaByteStream.h b/Engine/Plugins/Media/WmfMedia/Source/WmfMedia/Private/Wmf/WmfMediaByteStream.h
--- a/Engine/Plugins/Media/WmfMedia/Source/WmfMedia/Private/Wmf/WmfMediaByteStream.h
+++ b/Engine/Plugins/Media/WmfMedia/Source/WmfMedia/Private/Wmf/WmfMediaByteStream.h
@@ -19,7 +19,7 @@ public:
 	 *
 	 * @param Buffer The buffer to stream from.
 	 */
-	FWmfMediaByteStream( const TSharedRef<TArray<uint8>>& InBuffer );
+	FWmfMediaByteStream( const TSharedRef<TArray<uint8>, ESPMode::ThreadSafe>& InBuffer );
 
 	/** Virtual destructor. */
 	virtual ~FWmfMediaByteStream() { }
@@ -60,7 +60,7 @@ private:
 	bool AsyncReadInProgress;
 
 	/** Holds the buffer to stream from. */
-	TSharedRef<TArray<uint8>> Buffer;
+	TSharedRef<TArray<uint8>, ESPMode::ThreadSafe> Buffer;
 
 	/** Critical section for locking access to this class. */
 	mutable FCriticalSection CriticalSection;
diff --git a/Engine/Source/Runtime/Media/Public/IMediaPlayer.h b/Engine/Source/Runtime/Media/Public/IMediaPlayer.h
--- a/Engine/Source/Runtime/Media/Public/IMediaPlayer.h
+++ b/Engine/Source/Runtime/Media/Public/IMediaPlayer.h
@@ -137,7 +137,7 @@ public:
 	 * @return true if the media was opened, false otherwise.
 	 * @see Close, IsReady
 	 */
-	virtual bool Open( const TSharedRef<TArray<uint8>>& Buffer, const FString& OriginalUrl ) = 0;
+	virtual bool Open( const TSharedRef<TArray<uint8>, ESPMode::ThreadSafe>& Buffer, const FString& OriginalUrl ) = 0;
 
 	/**
 	 * Changes the media's playback time.
