Media: Refactored IMediaPlayer.Open to take an FArchive instead of byte array; removed pre-caching of media into memory

#UpgradeNotes: This is an API breaking change. If you implemented a plug-in for the Media Framework then you have to update your media player's Open() method to match the new signature. Take a look at the changes in the WmfMedia plug-in to see the differences between loading from a buffer and an archive.

diff --git a/Engine/Plugins/Media/AndroidMedia/Source/AndroidMedia/Private/AndroidMediaPlayer.cpp b/Engine/Plugins/Media/AndroidMedia/Source/AndroidMedia/Private/AndroidMediaPlayer.cpp
--- a/Engine/Plugins/Media/AndroidMedia/Source/AndroidMedia/Private/AndroidMediaPlayer.cpp
+++ b/Engine/Plugins/Media/AndroidMedia/Source/AndroidMedia/Private/AndroidMediaPlayer.cpp
@@ -533,7 +533,7 @@ bool FAndroidMediaPlayer::Open(const FString& Url)
 	return MediaState == EMediaState::Prepared;
 }
 
-bool FAndroidMediaPlayer::Open(const TSharedRef<TArray<uint8>, ESPMode::ThreadSafe>& Buffer,
+bool FAndroidMediaPlayer::Open(const TSharedRef<FArchive, ESPMode::ThreadSafe>& Archive,
 	const FString& OriginalUrl)
 {
 	return false;
diff --git a/Engine/Plugins/Media/AndroidMedia/Source/AndroidMedia/Private/AndroidMediaPlayer.h b/Engine/Plugins/Media/AndroidMedia/Source/AndroidMedia/Private/AndroidMediaPlayer.h
--- a/Engine/Plugins/Media/AndroidMedia/Source/AndroidMedia/Private/AndroidMediaPlayer.h
+++ b/Engine/Plugins/Media/AndroidMedia/Source/AndroidMedia/Private/AndroidMediaPlayer.h
@@ -41,7 +41,7 @@ public:
 	virtual bool IsPlaying() const override;
 	virtual bool IsReady() const override;
 	virtual bool Open(const FString& Url) override;
-	virtual bool Open(const TSharedRef<TArray<uint8>, ESPMode::ThreadSafe>& Buffer,
+	virtual bool Open(const TSharedRef<FArchive, ESPMode::ThreadSafe>& Archive,
 		const FString& OriginalUrl) override;
 	virtual bool Seek(const FTimespan& Time) override;
 	virtual bool SetLooping(bool Looping) override;
diff --git a/Engine/Plugins/Media/AvfMedia/Source/AvfMedia/Private/Player/AvfMediaPlayer.cpp b/Engine/Plugins/Media/AvfMedia/Source/AvfMedia/Private/Player/AvfMediaPlayer.cpp
--- a/Engine/Plugins/Media/AvfMedia/Source/AvfMedia/Private/Player/AvfMediaPlayer.cpp
+++ b/Engine/Plugins/Media/AvfMedia/Source/AvfMedia/Private/Player/AvfMediaPlayer.cpp
@@ -318,7 +318,7 @@ bool FAvfMediaPlayer::Open( const FString& Url )
 }
 
 
-bool FAvfMediaPlayer::Open( const TSharedRef<TArray<uint8>, ESPMode::ThreadSafe>& Buffer, const FString& OriginalUrl )
+bool FAvfMediaPlayer::Open( const TSharedRef<FArchive, ESPMode::ThreadSafe>& Archive, const FString& OriginalUrl )
 {
     return false;
 }
diff --git a/Engine/Plugins/Media/AvfMedia/Source/AvfMedia/Private/Player/AvfMediaPlayer.h b/Engine/Plugins/Media/AvfMedia/Source/AvfMedia/Private/Player/AvfMediaPlayer.h
--- a/Engine/Plugins/Media/AvfMedia/Source/AvfMedia/Private/Player/AvfMediaPlayer.h
+++ b/Engine/Plugins/Media/AvfMedia/Source/AvfMedia/Private/Player/AvfMediaPlayer.h
@@ -67,7 +67,7 @@ public:
 	virtual bool IsPlaying() const override;
 	virtual bool IsReady() const override;
 	virtual bool Open( const FString& Url ) override;
-	virtual bool Open( const TSharedRef<TArray<uint8>, ESPMode::ThreadSafe>& Buffer, const FString& OriginalUrl ) override;
+	virtual bool Open( const TSharedRef<FArchive, ESPMode::ThreadSafe>& Archive, const FString& OriginalUrl ) override;
 	virtual bool Seek( const FTimespan& Time ) override;
 	virtual bool SetLooping( bool Looping ) override;
 	virtual bool SetRate( float Rate ) override;
diff --git a/Engine/Plugins/Media/WmfMedia/Source/WmfMedia/Private/Player/WmfMediaPlayer.cpp b/Engine/Plugins/Media/WmfMedia/Source/WmfMedia/Private/Player/WmfMediaPlayer.cpp
--- a/Engine/Plugins/Media/WmfMedia/Source/WmfMedia/Private/Player/WmfMediaPlayer.cpp
+++ b/Engine/Plugins/Media/WmfMedia/Source/WmfMedia/Private/Player/WmfMediaPlayer.cpp
@@ -28,7 +28,7 @@ FTimespan FWmfMediaPlayer::GetDuration() const
 }
 
 
-TRange<float> FWmfMediaPlayer::GetSupportedRates( EMediaPlaybackDirections Direction, bool Unthinned ) const
+TRange<float> FWmfMediaPlayer::GetSupportedRates(EMediaPlaybackDirections Direction, bool Unthinned) const
 {
 	if (MediaSession != NULL)
 	{
@@ -159,7 +159,7 @@ bool FWmfMediaPlayer::IsReady() const
 }
 
 
-bool FWmfMediaPlayer::Open( const FString& Url )
+bool FWmfMediaPlayer::Open(const FString& Url)
 {
 	if (Url.IsEmpty())
 	{
@@ -188,9 +188,9 @@ bool FWmfMediaPlayer::Open( const FString& Url )
 }
 
 
-bool FWmfMediaPlayer::Open( const TSharedRef<TArray<uint8>, ESPMode::ThreadSafe>& Buffer, const FString& OriginalUrl )
+bool FWmfMediaPlayer::Open(const TSharedRef<FArchive, ESPMode::ThreadSafe>& Archive, const FString& OriginalUrl)
 {
-	if ((Buffer->Num() == 0) || OriginalUrl.IsEmpty())
+	if ((Archive->TotalSize() == 0) || OriginalUrl.IsEmpty())
 	{
 		return false;
 	}
@@ -203,7 +203,7 @@ bool FWmfMediaPlayer::Open( const TSharedRef<TArray<uint8>, ESPMode::ThreadSafe>
 	}
 
 	// create the media source from the given buffer
-	TComPtr<FWmfMediaByteStream> ByteStream = new FWmfMediaByteStream(Buffer);
+	TComPtr<FWmfMediaByteStream> ByteStream = new FWmfMediaByteStream(Archive);
 	MF_OBJECT_TYPE ObjectType = MF_OBJECT_INVALID;
 	TComPtr<IUnknown> SourceObject;
 	
@@ -218,13 +218,13 @@ bool FWmfMediaPlayer::Open( const TSharedRef<TArray<uint8>, ESPMode::ThreadSafe>
 }
 
 
-bool FWmfMediaPlayer::Seek( const FTimespan& Time )
+bool FWmfMediaPlayer::Seek(const FTimespan& Time)
 {
 	return ((MediaSession != NULL) && MediaSession->SetPosition(Time));
 }
 
 
-bool FWmfMediaPlayer::SetLooping( bool Looping )
+bool FWmfMediaPlayer::SetLooping(bool Looping)
 {
 	if (MediaSession == NULL)
 	{
@@ -237,7 +237,7 @@ bool FWmfMediaPlayer::SetLooping( bool Looping )
 }
 
 
-bool FWmfMediaPlayer::SetRate( float Rate )
+bool FWmfMediaPlayer::SetRate(float Rate)
 {
 	if (MediaSession == NULL)
 	{
@@ -261,7 +261,7 @@ bool FWmfMediaPlayer::SetRate( float Rate )
 /* FWmfMediaPlayer implementation
  *****************************************************************************/
 
-void FWmfMediaPlayer::AddStreamToTopology( uint32 StreamIndex, IMFTopology* Topology, IMFPresentationDescriptor* PresentationDescriptor, IMFMediaSource* MediaSourceObject )
+void FWmfMediaPlayer::AddStreamToTopology(uint32 StreamIndex, IMFTopology* Topology, IMFPresentationDescriptor* PresentationDescriptor, IMFMediaSource* MediaSourceObject)
 {
 	// get stream descriptor
 	TComPtr<IMFStreamDescriptor> StreamDescriptor;
@@ -400,7 +400,7 @@ void FWmfMediaPlayer::AddStreamToTopology( uint32 StreamIndex, IMFTopology* Topo
 }
 
 
-bool FWmfMediaPlayer::InitializeMediaSession( IUnknown* SourceObject, const FString& SourceUrl )
+bool FWmfMediaPlayer::InitializeMediaSession(IUnknown* SourceObject, const FString& SourceUrl)
 {
 	Close();
 
diff --git a/Engine/Plugins/Media/WmfMedia/Source/WmfMedia/Private/Player/WmfMediaPlayer.h b/Engine/Plugins/Media/WmfMedia/Source/WmfMedia/Private/Player/WmfMediaPlayer.h
--- a/Engine/Plugins/Media/WmfMedia/Source/WmfMedia/Private/Player/WmfMediaPlayer.h
+++ b/Engine/Plugins/Media/WmfMedia/Source/WmfMedia/Private/Player/WmfMediaPlayer.h
@@ -25,9 +25,9 @@ public:
 	// IMediaInfo interface
 
 	virtual FTimespan GetDuration() const override;
-	virtual TRange<float> GetSupportedRates( EMediaPlaybackDirections Direction, bool Unthinned ) const override;
+	virtual TRange<float> GetSupportedRates(EMediaPlaybackDirections Direction, bool Unthinned) const override;
 	virtual FString GetUrl() const override;
-	virtual bool SupportsRate( float Rate, bool Unthinned ) const override;
+	virtual bool SupportsRate(float Rate, bool Unthinned) const override;
 	virtual bool SupportsScrubbing() const override;
 	virtual bool SupportsSeeking() const override;
 
@@ -44,11 +44,11 @@ public:
 	virtual bool IsPaused() const override;
 	virtual bool IsPlaying() const override;
 	virtual bool IsReady() const override;
-	virtual bool Open( const FString& Url ) override;
-	virtual bool Open( const TSharedRef<TArray<uint8>, ESPMode::ThreadSafe>& Buffer, const FString& OriginalUrl ) override;
-	virtual bool Seek( const FTimespan& Time ) override;
-	virtual bool SetLooping( bool Looping ) override;
-	virtual bool SetRate( float Rate ) override;
+	virtual bool Open(const FString& Url) override;
+	virtual bool Open(const TSharedRef<FArchive, ESPMode::ThreadSafe>& Archive, const FString& OriginalUrl) override;
+	virtual bool Seek(const FTimespan& Time) override;
+	virtual bool SetLooping(bool Looping) override;
+	virtual bool SetRate(float Rate) override;
 
 	DECLARE_DERIVED_EVENT(FWmfMediaPlayer, IMediaPlayer::FOnMediaClosed, FOnMediaClosed);
 	virtual FOnMediaClosed& OnClosed() override
@@ -73,7 +73,7 @@ protected:
 	 * @param MediaSourceObject The media source object.
 	 * @return true on success, false otherwise.
 	 */
-	void AddStreamToTopology( uint32 StreamIndex, IMFTopology* Topology, IMFPresentationDescriptor* PresentationDescriptor, IMFMediaSource* MediaSourceObject );
+	void AddStreamToTopology(uint32 StreamIndex, IMFTopology* Topology, IMFPresentationDescriptor* PresentationDescriptor, IMFMediaSource* MediaSourceObject);
 
 	/**
 	 * Initializes the media session for the given media source.
@@ -82,7 +82,7 @@ protected:
 	 * @param SourceUrl The original URL of the media source.
 	 * @return true on success, false otherwise.
 	 */
-	bool InitializeMediaSession( IUnknown* SourceObject, const FString& SourceUrl );
+	bool InitializeMediaSession(IUnknown* SourceObject, const FString& SourceUrl);
 
 private:
 
diff --git a/Engine/Plugins/Media/WmfMedia/Source/WmfMedia/Private/Wmf/WmfMediaByteStream.cpp b/Engine/Plugins/Media/WmfMedia/Source/WmfMedia/Private/Wmf/WmfMediaByteStream.cpp
--- a/Engine/Plugins/Media/WmfMedia/Source/WmfMedia/Private/Wmf/WmfMediaByteStream.cpp
+++ b/Engine/Plugins/Media/WmfMedia/Source/WmfMedia/Private/Wmf/WmfMediaByteStream.cpp
@@ -7,10 +7,9 @@
 /* FWmfByteStream structors
  *****************************************************************************/
 
-FWmfMediaByteStream::FWmfMediaByteStream( const TSharedRef<TArray<uint8>, ESPMode::ThreadSafe>& InBuffer )
+FWmfMediaByteStream::FWmfMediaByteStream(const TSharedRef<FArchive, ESPMode::ThreadSafe>& InArchive)
 	: AsyncReadInProgress(false)
-	, Buffer(InBuffer)
-	, Position(0)
+	, Archive(InArchive)
 	, RefCount(0)
 { }
 
@@ -24,13 +23,13 @@ STDMETHODIMP_(ULONG) FWmfMediaByteStream::AddRef()
 }
 
 
-STDMETHODIMP FWmfMediaByteStream::GetParameters( unsigned long*, unsigned long*)
+STDMETHODIMP FWmfMediaByteStream::GetParameters(unsigned long*, unsigned long*)
 {
 	return E_NOTIMPL;
 }
 
 
-STDMETHODIMP FWmfMediaByteStream::Invoke( IMFAsyncResult* AsyncResult )
+STDMETHODIMP FWmfMediaByteStream::Invoke(IMFAsyncResult* AsyncResult)
 {
 	// recover read state
 	IUnknown* State = NULL;
@@ -83,7 +82,7 @@ STDMETHODIMP FWmfMediaByteStream::Invoke( IMFAsyncResult* AsyncResult )
 }
 
 
-STDMETHODIMP FWmfMediaByteStream::QueryInterface( REFIID RefID, void** Object )
+STDMETHODIMP FWmfMediaByteStream::QueryInterface(REFIID RefID, void** Object)
 {
 	static const QITAB QITab[] =
 	{
@@ -112,7 +111,7 @@ STDMETHODIMP_(ULONG) FWmfMediaByteStream::Release()
 /* IMFByteStream interface
  *****************************************************************************/
 
-STDMETHODIMP FWmfMediaByteStream::BeginRead( BYTE* pb, ULONG cb, IMFAsyncCallback* pCallback, IUnknown* punkState )
+STDMETHODIMP FWmfMediaByteStream::BeginRead(BYTE* pb, ULONG cb, IMFAsyncCallback* pCallback, IUnknown* punkState)
 {
 	if ((pCallback == NULL) || (pb == NULL))
 	{
@@ -141,7 +140,7 @@ STDMETHODIMP FWmfMediaByteStream::BeginRead( BYTE* pb, ULONG cb, IMFAsyncCallbac
 }
 
 
-STDMETHODIMP FWmfMediaByteStream::BeginWrite( const BYTE* pb, ULONG cb, IMFAsyncCallback* pCallback, IUnknown* punkState )
+STDMETHODIMP FWmfMediaByteStream::BeginWrite(const BYTE* pb, ULONG cb, IMFAsyncCallback* pCallback, IUnknown* punkState)
 {
 	return E_NOTIMPL;
 }
@@ -153,7 +152,7 @@ STDMETHODIMP FWmfMediaByteStream::Close()
 }
 
 
-STDMETHODIMP FWmfMediaByteStream::EndRead( IMFAsyncResult* pResult, ULONG* pcbRead )
+STDMETHODIMP FWmfMediaByteStream::EndRead(IMFAsyncResult* pResult, ULONG* pcbRead)
 {
 	if (pcbRead == NULL)
 	{
@@ -175,7 +174,7 @@ STDMETHODIMP FWmfMediaByteStream::EndRead( IMFAsyncResult* pResult, ULONG* pcbRe
 }
 
 
-STDMETHODIMP FWmfMediaByteStream::EndWrite( IMFAsyncResult* pResult, ULONG* pcbWritten )
+STDMETHODIMP FWmfMediaByteStream::EndWrite(IMFAsyncResult* pResult, ULONG* pcbWritten)
 {
 	return E_NOTIMPL;
 }
@@ -187,7 +186,7 @@ STDMETHODIMP FWmfMediaByteStream::Flush()
 }
 
 
-STDMETHODIMP FWmfMediaByteStream::GetCapabilities( DWORD* pdwCapabilities )
+STDMETHODIMP FWmfMediaByteStream::GetCapabilities(DWORD* pdwCapabilities)
 {
 	FScopeLock ScopeLock(&CriticalSection);
 	*pdwCapabilities = MFBYTESTREAM_IS_READABLE | MFBYTESTREAM_IS_SEEKABLE;
@@ -196,25 +195,25 @@ STDMETHODIMP FWmfMediaByteStream::GetCapabilities( DWORD* pdwCapabilities )
 }
 
 
-STDMETHODIMP FWmfMediaByteStream::GetCurrentPosition( QWORD* pqwPosition )
+STDMETHODIMP FWmfMediaByteStream::GetCurrentPosition(QWORD* pqwPosition)
 {
 	FScopeLock ScopeLock(&CriticalSection);
-	*pqwPosition = Position;
+	*pqwPosition = Archive->Tell();
 
 	return S_OK;
 }
 
 
-STDMETHODIMP FWmfMediaByteStream::GetLength( QWORD* pqwLength )
+STDMETHODIMP FWmfMediaByteStream::GetLength(QWORD* pqwLength)
 {
 	FScopeLock ScopeLock(&CriticalSection);
-	*pqwLength = (QWORD)Buffer->Num();
+	*pqwLength = (QWORD)Archive->TotalSize();
 
 	return S_OK;
 }
 
 
-STDMETHODIMP FWmfMediaByteStream::IsEndOfStream( BOOL* pfEndOfStream )
+STDMETHODIMP FWmfMediaByteStream::IsEndOfStream(BOOL* pfEndOfStream)
 {
 	if (pfEndOfStream == NULL)
 	{
@@ -222,31 +221,33 @@ STDMETHODIMP FWmfMediaByteStream::IsEndOfStream( BOOL* pfEndOfStream )
 	}
 
 	FScopeLock ScopeLock(&CriticalSection);
-	*pfEndOfStream = (Position >= Buffer->Num()) ? TRUE : FALSE;
+	*pfEndOfStream = Archive->AtEnd() ? TRUE : FALSE;
 
 	return S_OK;
 }
 
 
-STDMETHODIMP FWmfMediaByteStream::Read( BYTE* pb, ULONG cb, ULONG* pcbRead )
+STDMETHODIMP FWmfMediaByteStream::Read(BYTE* pb, ULONG cb, ULONG* pcbRead)
 {
 	FScopeLock ScopeLock(&CriticalSection);
 
+	int64 Position = Archive->Tell();
+	int64 Size = Archive->TotalSize();
 	ULONG BytesToRead = cb;
 
-	if (BytesToRead > (ULONG)Buffer->Num())
+	if (BytesToRead > (ULONG)Size)
 	{
-		BytesToRead = Buffer->Num();
+		BytesToRead = Size;
 	}
 
-	if ((Buffer->Num() - BytesToRead) < Position)
+	if ((Size - BytesToRead) < Archive->Tell())
 	{
-		BytesToRead = Buffer->Num() - Position;
+		BytesToRead = Size - Position;
 	}
 
 	if (BytesToRead > 0)
 	{
-		FMemory::Memcpy(pb, ((BYTE*)Buffer->GetData() + Position), BytesToRead);
+		Archive->Serialize(pb, BytesToRead);
 	}
 
 	if (pcbRead != NULL)
@@ -254,13 +255,13 @@ STDMETHODIMP FWmfMediaByteStream::Read( BYTE* pb, ULONG cb, ULONG* pcbRead )
 		*pcbRead = BytesToRead;
 	}	
 
-	Position += BytesToRead;
+	Archive->Seek(Position + BytesToRead);
 
 	return S_OK;
 }
 
 
-STDMETHODIMP FWmfMediaByteStream::Seek( MFBYTESTREAM_SEEK_ORIGIN SeekOrigin, LONGLONG qwSeekOffset, DWORD dwSeekFlags, QWORD* pqwCurrentPosition )
+STDMETHODIMP FWmfMediaByteStream::Seek(MFBYTESTREAM_SEEK_ORIGIN SeekOrigin, LONGLONG qwSeekOffset, DWORD dwSeekFlags, QWORD* pqwCurrentPosition)
 {
 	FScopeLock ScopeLock(&CriticalSection);
 
@@ -271,29 +272,29 @@ STDMETHODIMP FWmfMediaByteStream::Seek( MFBYTESTREAM_SEEK_ORIGIN SeekOrigin, LON
 
 	if (SeekOrigin == msoCurrent)
 	{
-		Position += qwSeekOffset;
+		Archive->Seek(Archive->Tell() + qwSeekOffset);
 	}
 	else
 	{
-		Position = qwSeekOffset;
+		Archive->Seek(qwSeekOffset);
 	}
 
 	if (pqwCurrentPosition != NULL)
 	{
-		*pqwCurrentPosition = (QWORD)Position;
+		*pqwCurrentPosition = (QWORD)Archive->Tell();
 	}
 
 	return S_OK;
 }
 
 
-STDMETHODIMP FWmfMediaByteStream::SetLength( QWORD qwLength )
+STDMETHODIMP FWmfMediaByteStream::SetLength(QWORD qwLength)
 {
 	return E_NOTIMPL;
 }
 
 
-STDMETHODIMP FWmfMediaByteStream::SetCurrentPosition( QWORD qwPosition )
+STDMETHODIMP FWmfMediaByteStream::SetCurrentPosition(QWORD qwPosition)
 {
 	FScopeLock ScopeLock(&CriticalSection);
 
@@ -302,13 +303,13 @@ STDMETHODIMP FWmfMediaByteStream::SetCurrentPosition( QWORD qwPosition )
 		return S_FALSE;
 	}
 
-	Position = (uint64)qwPosition;
+	Archive->Seek((int64)qwPosition);
 
 	return S_OK;
 }
 
 
-STDMETHODIMP FWmfMediaByteStream::Write( const BYTE* pb, ULONG cb, ULONG* pcbWritten )
+STDMETHODIMP FWmfMediaByteStream::Write(const BYTE* pb, ULONG cb, ULONG* pcbWritten)
 {
 	return E_NOTIMPL;
 }
diff --git a/Engine/Plugins/Media/WmfMedia/Source/WmfMedia/Private/Wmf/WmfMediaByteStream.h b/Engine/Plugins/Media/WmfMedia/Source/WmfMedia/Private/Wmf/WmfMediaByteStream.h
--- a/Engine/Plugins/Media/WmfMedia/Source/WmfMedia/Private/Wmf/WmfMediaByteStream.h
+++ b/Engine/Plugins/Media/WmfMedia/Source/WmfMedia/Private/Wmf/WmfMediaByteStream.h
@@ -17,9 +17,9 @@ public:
 	/**
 	 * Creates and initializes a new instance from the specified buffer.
 	 *
-	 * @param Buffer The buffer to stream from.
+	 * @param InArchive The archive to stream from.
 	 */
-	FWmfMediaByteStream( const TSharedRef<TArray<uint8>, ESPMode::ThreadSafe>& InBuffer );
+	FWmfMediaByteStream(const TSharedRef<FArchive, ESPMode::ThreadSafe>& InArchive);
 
 	/** Virtual destructor. */
 	virtual ~FWmfMediaByteStream() { }
@@ -29,45 +29,42 @@ public:
 	// IMFAsyncCallback interface
 
 	STDMETHODIMP_(ULONG) AddRef();
-	STDMETHODIMP GetParameters( unsigned long*, unsigned long*);
-	STDMETHODIMP Invoke( IMFAsyncResult* AsyncResult );
-	STDMETHODIMP QueryInterface( REFIID RefID, void** Object );
+	STDMETHODIMP GetParameters(unsigned long*, unsigned long*);
+	STDMETHODIMP Invoke(IMFAsyncResult* AsyncResult);
+	STDMETHODIMP QueryInterface(REFIID RefID, void** Object);
 	STDMETHODIMP_(ULONG) Release();
 
 public:
 
 	// IMFByteStream interface
 
-	STDMETHODIMP BeginRead( BYTE* pb, ULONG cb, IMFAsyncCallback* pCallback, IUnknown* punkState );
-	STDMETHODIMP BeginWrite( const BYTE* pb, ULONG cb, IMFAsyncCallback* pCallback, IUnknown* punkState );
+	STDMETHODIMP BeginRead(BYTE* pb, ULONG cb, IMFAsyncCallback* pCallback, IUnknown* punkState);
+	STDMETHODIMP BeginWrite(const BYTE* pb, ULONG cb, IMFAsyncCallback* pCallback, IUnknown* punkState);
 	STDMETHODIMP Close();
-	STDMETHODIMP EndRead( IMFAsyncResult* pResult, ULONG* pcbRead );
-	STDMETHODIMP EndWrite( IMFAsyncResult* pResult, ULONG* pcbWritten );
+	STDMETHODIMP EndRead( IMFAsyncResult* pResult, ULONG* pcbRead);
+	STDMETHODIMP EndWrite( IMFAsyncResult* pResult, ULONG* pcbWritten);
 	STDMETHODIMP Flush();
-	STDMETHODIMP GetCapabilities( DWORD* pdwCapabilities );
-	STDMETHODIMP GetCurrentPosition( QWORD* pqwPosition );
-	STDMETHODIMP GetLength( QWORD* pqwLength );
-	STDMETHODIMP IsEndOfStream( BOOL* pfEndOfStream );
-	STDMETHODIMP Read( BYTE* pb, ULONG cb, ULONG* pcbRead );
-	STDMETHODIMP Seek( MFBYTESTREAM_SEEK_ORIGIN SeekOrigin, LONGLONG qwSeekOffset, DWORD dwSeekFlags, QWORD* pqwCurrentPosition );
-	STDMETHODIMP SetLength( QWORD qwLength );
-	STDMETHODIMP SetCurrentPosition( QWORD qwPosition );
-	STDMETHODIMP Write( const BYTE* pb, ULONG cb, ULONG* pcbWritten );
+	STDMETHODIMP GetCapabilities(DWORD* pdwCapabilities);
+	STDMETHODIMP GetCurrentPosition(QWORD* pqwPosition);
+	STDMETHODIMP GetLength(QWORD* pqwLength);
+	STDMETHODIMP IsEndOfStream( BOOL* pfEndOfStream);
+	STDMETHODIMP Read(BYTE* pb, ULONG cb, ULONG* pcbRead);
+	STDMETHODIMP Seek(MFBYTESTREAM_SEEK_ORIGIN SeekOrigin, LONGLONG qwSeekOffset, DWORD dwSeekFlags, QWORD* pqwCurrentPosition);
+	STDMETHODIMP SetLength(QWORD qwLength);
+	STDMETHODIMP SetCurrentPosition(QWORD qwPosition);
+	STDMETHODIMP Write(const BYTE* pb, ULONG cb, ULONG* pcbWritten );
 
 private:
 
 	/** Whether the stream is currently being read asynchronously. */
 	bool AsyncReadInProgress;
 
-	/** Holds the buffer to stream from. */
-	TSharedRef<TArray<uint8>, ESPMode::ThreadSafe> Buffer;
+	/** Holds the archive to stream from. */
+	TSharedRef<FArchive, ESPMode::ThreadSafe> Archive;
 
 	/** Critical section for locking access to this class. */
 	mutable FCriticalSection CriticalSection;
 
-	/** The current read position. */
-	uint64 Position;
-
 	/** Holds a reference counter for this instance. */
 	int32 RefCount;
 };
diff --git a/Engine/Source/Runtime/Media/Public/IMediaPlayer.h b/Engine/Source/Runtime/Media/Public/IMediaPlayer.h
--- a/Engine/Source/Runtime/Media/Public/IMediaPlayer.h
+++ b/Engine/Source/Runtime/Media/Public/IMediaPlayer.h
@@ -130,14 +130,14 @@ public:
 	virtual bool Open( const FString& Url ) = 0;
 
 	/**
-	 * Opens a media from a buffer.
+	 * Opens a media from a file or memory archive.
 	 *
-	 * @param Buffer The buffer holding the media data.
+	 * @param Archive The archive holding the media data.
 	 * @param OriginalUrl The original URL of the media that was loaded into the buffer.
 	 * @return true if the media was opened, false otherwise.
 	 * @see Close, IsReady
 	 */
-	virtual bool Open( const TSharedRef<TArray<uint8>, ESPMode::ThreadSafe>& Buffer, const FString& OriginalUrl ) = 0;
+	virtual bool Open( const TSharedRef<FArchive, ESPMode::ThreadSafe>& Archive, const FString& OriginalUrl ) = 0;
 
 	/**
 	 * Changes the media's playback time.
diff --git a/Engine/Source/Runtime/MediaAssets/Private/Assets/MediaPlayer.cpp b/Engine/Source/Runtime/MediaAssets/Private/Assets/MediaPlayer.cpp
--- a/Engine/Source/Runtime/MediaAssets/Private/Assets/MediaPlayer.cpp
+++ b/Engine/Source/Runtime/MediaAssets/Private/Assets/MediaPlayer.cpp
@@ -11,7 +11,6 @@
 UMediaPlayer::UMediaPlayer( const FObjectInitializer& ObjectInitializer )
 	: Super(ObjectInitializer)
 	, Looping(true)
-	, StreamMode(MASM_FromUrl)
 	, Player(nullptr)
 { }
 
@@ -265,33 +264,13 @@ void UMediaPlayer::InitializePlayer()
 		}
 		else
 		{
-			const FString FullUrl = FPaths::ConvertRelativePathToFull(FPaths::IsRelative(URL) ? FPaths::GameContentDir() / URL : URL);
-
-			if (StreamMode == EMediaPlayerStreamModes::MASM_FromUrl)
-			{
-				OpenedSuccessfully = Player->Open(FullUrl);
-			}
-			else if (FPaths::FileExists(FullUrl))
-			{
-				FArchive* FileReader = IFileManager::Get().CreateFileReader(*FullUrl);
-		
-				if (FileReader == nullptr)
-				{
-					return;
-				}
-
-				if (FileReader->TotalSize() > 0)
-				{
-					TArray<uint8>* FileData = new TArray<uint8>();
-
-					FileData->AddUninitialized(FileReader->TotalSize());
-					FileReader->Serialize(FileData->GetData(), FileReader->TotalSize());
-
-					OpenedSuccessfully = Player->Open(MakeShareable(FileData), FullUrl);
-				}
-
-				delete FileReader;
-			}
+			const FString FullUrl = FPaths::ConvertRelativePathToFull(
+				FPaths::IsRelative(URL)
+					? FPaths::GameContentDir() / URL
+					: URL
+			);
+
+			OpenedSuccessfully = Player->Open(FullUrl);
 		}
 
 		// finish initialization
diff --git a/Engine/Source/Runtime/MediaAssets/Public/MediaPlayer.h b/Engine/Source/Runtime/MediaAssets/Public/MediaPlayer.h
--- a/Engine/Source/Runtime/MediaAssets/Public/MediaPlayer.h
+++ b/Engine/Source/Runtime/MediaAssets/Public/MediaPlayer.h
@@ -8,22 +8,6 @@
 class IMediaPlayer;
 
 
-/**
- * Enumerates available media streaming modes.
- */
-UENUM()
-enum EMediaPlayerStreamModes
-{
-	/** Load media contents to memory before playing. */
-	MASM_FromMemory UMETA(DisplayName="FromMemory"),
-
-	/** Play media directly from the URL. */
-	MASM_FromUrl UMETA(DisplayName="FromUrl"),
-
-	MASM_MAX,
-};
-
-
 /** Multicast delegate that is invoked when a media player's media has been closed. */
 DECLARE_DYNAMIC_MULTICAST_DELEGATE(FOnMediaPlayerMediaClosed);
 
@@ -298,10 +282,6 @@ protected:
 	UPROPERTY(EditAnywhere, Category=Playback)
 	uint32 Looping:1;
 
-	/** Select where to stream the media from, i.e. file or memory. */
-	UPROPERTY(EditAnywhere, BlueprintReadOnly, Category=Source)
-	TEnumAsByte<EMediaPlayerStreamModes> StreamMode;
-
 	/** The path or URL to the media file to be played. */
 	UPROPERTY(EditAnywhere, Category=Source)
 	FString URL;
