WmfMedia: Fixed playback stopping on MEEndOfPresentation instead of MESessionEnded

diff --git a/Engine/Plugins/Media/WmfMedia/Source/WmfMedia/Private/Wmf/WmfMediaSession.cpp b/Engine/Plugins/Media/WmfMedia/Source/WmfMedia/Private/Wmf/WmfMediaSession.cpp
--- a/Engine/Plugins/Media/WmfMedia/Source/WmfMedia/Private/Wmf/WmfMediaSession.cpp
+++ b/Engine/Plugins/Media/WmfMedia/Source/WmfMedia/Private/Wmf/WmfMediaSession.cpp
@@ -215,7 +215,11 @@ STDMETHODIMP FWmfMediaSession::Invoke( IMFAsyncResult* AsyncResult )
 			}
 			else
 			{
-				if (EventType == MEEndOfPresentation)
+				if (EventType == MESessionCapabilitiesChanged)
+				{
+					Capabilities = ::MFGetAttributeUINT32(Event, MF_EVENT_SESSIONCAPS, Capabilities);
+				}
+				else if (EventType == MESessionEnded)
 				{
 					if (Looping)
 					{
@@ -229,17 +233,9 @@ STDMETHODIMP FWmfMediaSession::Invoke( IMFAsyncResult* AsyncResult )
 					else
 					{
 						RequestedState = EMediaStates::Stopped;
-						MediaSession->Stop();
+						UpdateState(EMediaStates::Stopped);
 					}
 				}
-				else if (EventType == MESessionCapabilitiesChanged)
-				{
-					Capabilities = ::MFGetAttributeUINT32(Event, MF_EVENT_SESSIONCAPS, Capabilities);
-				}
-				else if (EventType == MESessionEnded)
-				{
-					// do nothing
-				}
 				else if (EventType == MESessionPaused)
 				{
 					UpdateState(EMediaStates::Paused);
