WmfMedia: Added/fixed code comments

diff --git a/Engine/Plugins/Media/WmfMedia/Source/WmfMedia/Private/Wmf/WmfMediaSession.h b/Engine/Plugins/Media/WmfMedia/Source/WmfMedia/Private/Wmf/WmfMediaSession.h
--- a/Engine/Plugins/Media/WmfMedia/Source/WmfMedia/Private/Wmf/WmfMediaSession.h
+++ b/Engine/Plugins/Media/WmfMedia/Source/WmfMedia/Private/Wmf/WmfMediaSession.h
@@ -56,9 +56,8 @@ public:
 	 */
 	FWmfMediaSession( const FTimespan& InDuration, const TComPtr<IMFTopology>& InTopology );
 
-	virtual ~FWmfMediaSession()
-	{
-	}
+	/** Virtual destructor. */
+	virtual ~FWmfMediaSession() { }
 
 public:
 
@@ -106,7 +105,7 @@ public:
 	 * @param Unthinned Whether the rates are for unthinned playback.
 	 * @return The supported playback rate.
 	 */
-	TRange<float> GetSupportedRates( EMediaPlaybackDirections Direction, bool Unthinned ) const;
+	TRange<float> GetSupportedRates(EMediaPlaybackDirections Direction, bool Unthinned) const;
 
 	/**
 	 * Checks whether playback is currently looping.
@@ -125,14 +124,14 @@ public:
 	 * @param Unthinned Whether no frames should be dropped at the given rate.
 	 * @return true if the rate is supported, false otherwise.
 	 */
-	bool IsRateSupported( float Rate, bool Unthinned ) const;
+	bool IsRateSupported(float Rate, bool Unthinned) const;
 
 	/**
 	 * Sets whether playback should loop back to the beginning.
 	 *
 	 * @param InLooping Whether playback should be looped.
 	 */
-	void SetLooping( bool InLooping )
+	void SetLooping(bool InLooping)
 	{
 		Looping = InLooping;
 	}
@@ -143,7 +142,7 @@ public:
 	 * @param Position The position to set.
 	 * @return true if the position will be changed, false otherwise.
 	 */
-	bool SetPosition( const FTimespan& Position );
+	bool SetPosition(const FTimespan& Position);
 
 	/**
 	 * Sets the playback rate.
@@ -151,7 +150,7 @@ public:
 	 * @param Rate The rate to set.
 	 * @return true if the rate will be changed, false otherwise.
 	 */
-	bool SetRate( float Rate );
+	bool SetRate(float Rate);
 
 	/**
 	 * Sets the media state.
@@ -159,7 +158,7 @@ public:
 	 * @param NewState The media state to set.
 	 * @return true if the state will be changed, false otherwise.
 	 */
-	bool SetState( EMediaStates NewState );
+	bool SetState(EMediaStates NewState);
 
 	/**
 	 * Checks whether this session supports scrubbing.
@@ -185,9 +184,9 @@ public:
 	// IMFAsyncCallback interface
 
 	STDMETHODIMP_(ULONG) AddRef();
-	STDMETHODIMP GetParameters( unsigned long*, unsigned long*);
-	STDMETHODIMP Invoke( IMFAsyncResult* AsyncResult );
-	STDMETHODIMP QueryInterface( REFIID RefID, void** Object );
+	STDMETHODIMP GetParameters(unsigned long*, unsigned long*);
+	STDMETHODIMP Invoke(IMFAsyncResult* AsyncResult);
+	STDMETHODIMP QueryInterface(REFIID RefID, void** Object);
 	STDMETHODIMP_(ULONG) Release();
 
 protected:
@@ -211,7 +210,7 @@ protected:
 	 *
 	 * @param CompletedState The state that was completed.
 	 */
-	void UpdateState( EMediaStates CompletedState );
+	void UpdateState(EMediaStates CompletedState);
 
 private:
 
@@ -268,7 +267,7 @@ private:
 
 private:
 
-	/** Holds an event delegate that is invoked when an error occured. */
+	/** Holds an event delegate that is invoked when an error occurred. */
 	FOnError ErrorEvent;
 };
 
