WmfMedia: Fixed bad ref count in read state

diff --git a/Engine/Plugins/Media/WmfMedia/Source/WmfMedia/Private/Wmf/WmfMediaByteStream.cpp b/Engine/Plugins/Media/WmfMedia/Source/WmfMedia/Private/Wmf/WmfMediaByteStream.cpp
--- a/Engine/Plugins/Media/WmfMedia/Source/WmfMedia/Private/Wmf/WmfMediaByteStream.cpp
+++ b/Engine/Plugins/Media/WmfMedia/Source/WmfMedia/Private/Wmf/WmfMediaByteStream.cpp
@@ -125,7 +125,7 @@ STDMETHODIMP FWmfMediaByteStream::BeginRead(BYTE* pb, ULONG cb, IMFAsyncCallback
 		return E_INVALIDARG;
 	}
 
-	FWmfMediaReadState* ReadState = new(std::nothrow) FWmfMediaReadState(pb, cb);
+	TComPtr<FWmfMediaReadState> ReadState = new(std::nothrow) FWmfMediaReadState(pb, cb);
 
 	if (ReadState == NULL)
 	{
@@ -134,7 +134,6 @@ STDMETHODIMP FWmfMediaByteStream::BeginRead(BYTE* pb, ULONG cb, IMFAsyncCallback
 
 	IMFAsyncResult* AsyncResult = NULL;
 	HRESULT Result = ::MFCreateAsyncResult(ReadState, pCallback, punkState, &AsyncResult);
-	ReadState->Release();
 
 	if (SUCCEEDED(Result))
 	{
