Media: Added support for async opening of media; implemented async opening for WmfMedia

https://github.com/EpicGames/UnrealEngine/pull/1515

diff --git a/Engine/Plugins/Media/AndroidMedia/Source/AndroidMedia/Private/AndroidMediaPlayer.h b/Engine/Plugins/Media/AndroidMedia/Source/AndroidMedia/Private/AndroidMediaPlayer.h
--- a/Engine/Plugins/Media/AndroidMedia/Source/AndroidMedia/Private/AndroidMediaPlayer.h
+++ b/Engine/Plugins/Media/AndroidMedia/Source/AndroidMedia/Private/AndroidMediaPlayer.h
@@ -65,6 +65,12 @@ public:
 		return OpenedEvent;
 	}
 
+	DECLARE_DERIVED_EVENT(FAndroidMediaPlayer, IMediaPlayer::FOnMediaOpenFailed, FOnMediaOpenFailed);
+	virtual FOnMediaOpenFailed& OnOpenFailed() override
+	{
+		return OpenFailedEvent;
+	}
+
 	DECLARE_DERIVED_EVENT(FWmfMediaPlayer, IMediaPlayer::FOnTracksChanged, FOnTracksChanged);
 	virtual FOnTracksChanged& OnTracksChanged() override
 	{
@@ -110,6 +116,9 @@ private:
 	/** Holds an event delegate that is invoked when media has been opened. */
 	FOnMediaOpened OpenedEvent;
 
+	/** Holds an event delegate that is invoked when media failed to open. */
+	FOnMediaOpenFailed OpenFailedEvent;
+
 	/** Holds an event delegate that is invoked when the media tracks have changed. */
 	FOnTracksChanged TracksChangedEvent;
 
diff --git a/Engine/Plugins/Media/AvfMedia/Source/AvfMedia/Private/Player/AvfMediaPlayer.h b/Engine/Plugins/Media/AvfMedia/Source/AvfMedia/Private/Player/AvfMediaPlayer.h
--- a/Engine/Plugins/Media/AvfMedia/Source/AvfMedia/Private/Player/AvfMediaPlayer.h
+++ b/Engine/Plugins/Media/AvfMedia/Source/AvfMedia/Private/Player/AvfMediaPlayer.h
@@ -86,6 +86,12 @@ public:
 		return OpenedEvent;
 	}
 
+	DECLARE_DERIVED_EVENT(FAvfMediaPlayer, IMediaPlayer::FOnMediaOpenFailed, FOnMediaOpenFailed);
+	virtual FOnMediaOpenFailed& OnOpenFailed() override
+	{
+		return OpenFailedEvent;
+	}
+
 	DECLARE_DERIVED_EVENT(FWmfMediaPlayer, IMediaPlayer::FOnTracksChanged, FOnTracksChanged);
 	virtual FOnTracksChanged& OnTracksChanged() override
 	{
@@ -154,6 +160,9 @@ private:
 	/** Holds an event delegate that is invoked when media has been loaded. */
 	FOnMediaOpened OpenedEvent;
 
+	/** Holds an event delegate that is invoked when media failed to open. */
+	FOnMediaOpenFailed OpenFailedEvent;
+
 	/** Holds an event delegate that is invoked when the media tracks have changed. */
 	FOnTracksChanged TracksChangedEvent;
 };
diff --git a/Engine/Plugins/Media/WmfMedia/Source/WmfMedia/Private/Player/WmfMediaPlayer.cpp b/Engine/Plugins/Media/WmfMedia/Source/WmfMedia/Private/Player/WmfMediaPlayer.cpp
--- a/Engine/Plugins/Media/WmfMedia/Source/WmfMedia/Private/Player/WmfMediaPlayer.cpp
+++ b/Engine/Plugins/Media/WmfMedia/Source/WmfMedia/Private/Player/WmfMediaPlayer.cpp
@@ -2,6 +2,7 @@
 
 #include "WmfMediaPrivatePCH.h"
 #include "AllowWindowsPlatformTypes.h"
+#include "Async/Async.h"
 
 
 /* FWmfVideoPlayer structors
@@ -10,12 +11,30 @@
 FWmfMediaPlayer::FWmfMediaPlayer()
 	: Duration(0)
 	, MediaSession(nullptr)
-{ }
+{
+	Resolver = new(std::nothrow) FWmfMediaResolver;
+	{
+		Resolver->OnResolveComplete().BindLambda([=](TComPtr<IUnknown> SourceObject, FString ResolvedUrl) {
+			AsyncTask(ENamedThreads::GameThread, [=]() {
+				InitializeMediaSession(SourceObject, ResolvedUrl)
+					? OpenedEvent.Broadcast(ResolvedUrl)
+					: OpenFailedEvent.Broadcast(ResolvedUrl);
+			});
+		});
+
+		Resolver->OnResolveFailed().BindLambda([=](FString FailedUrl) {
+			AsyncTask(ENamedThreads::GameThread, [=]() {
+				OpenFailedEvent.Broadcast(FailedUrl);
+			});
+		});
+	}
+}
 
 
 FWmfMediaPlayer::~FWmfMediaPlayer()
 {
 	Close();
+	Resolver->OnResolveComplete().Unbind();
 }
 
 
@@ -72,6 +91,8 @@ bool FWmfMediaPlayer::SupportsSeeking() const
 
 void FWmfMediaPlayer::Close()
 {
+	Resolver->Cancel();
+
 	if (MediaSession == NULL)
 	{
 		return;
@@ -181,25 +202,7 @@ bool FWmfMediaPlayer::Open(const FString& Url)
 		return false;
 	}
 
-	TComPtr<IMFSourceResolver> SourceResolver;
-
-	if (FAILED(::MFCreateSourceResolver(&SourceResolver)))
-	{
-		return false;
-	}
-
-	// resolve the media source from the given URL
-	MF_OBJECT_TYPE ObjectType = MF_OBJECT_INVALID;
-	TComPtr<IUnknown> SourceObject;
-
-	if (FAILED(SourceResolver->CreateObjectFromURL(*Url, MF_RESOLUTION_MEDIASOURCE, NULL, &ObjectType, &SourceObject)))
-	{
-		UE_LOG(LogWmfMedia, Error, TEXT("Failed to create source object from URL for %s"), *Url);
-
-		return false;
-	}
-
-	return InitializeMediaSession(SourceObject, Url);
+	return Resolver->ResolveUrl(Url);
 }
 
 
@@ -210,26 +213,7 @@ bool FWmfMediaPlayer::Open(const TSharedRef<FArchive, ESPMode::ThreadSafe>& Arch
 		return false;
 	}
 
-	TComPtr<IMFSourceResolver> SourceResolver;
-	
-	if (FAILED(::MFCreateSourceResolver(&SourceResolver)))
-	{
-		return false;
-	}
-
-	// create the media source from the given buffer
-	TComPtr<FWmfMediaByteStream> ByteStream = new FWmfMediaByteStream(Archive);
-	MF_OBJECT_TYPE ObjectType = MF_OBJECT_INVALID;
-	TComPtr<IUnknown> SourceObject;
-	
-	if (FAILED(SourceResolver->CreateObjectFromByteStream(ByteStream, *OriginalUrl, MF_RESOLUTION_MEDIASOURCE, NULL, &ObjectType, &SourceObject)))
-	{
-		UE_LOG(LogWmfMedia, Error, TEXT("Failed to create source object from buffer for %s"), *OriginalUrl);
-
-		return false;
-	}
-
-	return InitializeMediaSession(SourceObject, OriginalUrl);
+	return Resolver->ResolveByteStream(Archive, OriginalUrl);
 }
 
 
@@ -484,8 +468,6 @@ bool FWmfMediaPlayer::InitializeMediaSession(IUnknown* SourceObject, const FStri
 		MediaSession->OnError().AddRaw(this, &FWmfMediaPlayer::HandleSessionError);
 	}
 
-	OpenedEvent.Broadcast(SourceUrl);
-
 	return (MediaSession->GetState() != EMediaStates::Error);
 }
 
diff --git a/Engine/Plugins/Media/WmfMedia/Source/WmfMedia/Private/Player/WmfMediaPlayer.h b/Engine/Plugins/Media/WmfMedia/Source/WmfMedia/Private/Player/WmfMediaPlayer.h
--- a/Engine/Plugins/Media/WmfMedia/Source/WmfMedia/Private/Player/WmfMediaPlayer.h
+++ b/Engine/Plugins/Media/WmfMedia/Source/WmfMedia/Private/Player/WmfMediaPlayer.h
@@ -64,6 +64,12 @@ public:
 		return OpenedEvent;
 	}
 
+	DECLARE_DERIVED_EVENT(FWmfMediaPlayer, IMediaPlayer::FOnMediaOpenFailed, FOnMediaOpenFailed);
+	virtual FOnMediaOpenFailed& OnOpenFailed() override
+	{
+		return OpenFailedEvent;
+	}
+
 	DECLARE_DERIVED_EVENT(FWmfMediaPlayer, IMediaPlayer::FOnTracksChanged, FOnTracksChanged);
 	virtual FOnTracksChanged& OnTracksChanged() override
 	{
@@ -117,6 +123,9 @@ private:
 	/** The URL of the currently opened media. */
 	FString MediaUrl;
 
+	/** The media source resolver. */
+	TComPtr<FWmfMediaResolver> Resolver;
+
 	/** The available video tracks. */
 	TArray<IMediaVideoTrackRef> VideoTracks;
 
@@ -128,6 +137,9 @@ private:
 	/** Holds an event delegate that is invoked when media has been opened. */
 	FOnMediaOpened OpenedEvent;
 
+	/** Holds an event delegate that is invoked when media failed to open. */
+	FOnMediaOpenFailed OpenFailedEvent;
+
 	/** Holds an event delegate that is invoked when the media tracks have changed. */
 	FOnTracksChanged TracksChangedEvent;
 };
diff --git a/Engine/Plugins/Media/WmfMedia/Source/WmfMedia/Private/Wmf/WmfMediaResolveState.h b/Engine/Plugins/Media/WmfMedia/Source/WmfMedia/Private/Wmf/WmfMediaResolveState.h
new file mode 100644
--- /dev/null
+++ b/Engine/Plugins/Media/WmfMedia/Source/WmfMedia/Private/Wmf/WmfMediaResolveState.h
@@ -0,0 +1,89 @@
+// Copyright 1998-2015 Epic Games, Inc. All Rights Reserved.
+
+#pragma once
+
+#include "AllowWindowsPlatformTypes.h"
+
+
+/** Enumerates media source resolve types. */
+enum class EWmfMediaResolveType
+{
+	/** Not resolving right now. */
+	None,
+
+	/** Resolving a byte stream. */
+	ByteStream,
+
+	/** Resolving a file or web URL. */
+	Url
+};
+
+
+/**
+ * Implements media source resolve state information.
+ */
+struct FWmfMediaResolveState
+	: public IUnknown
+{
+	/** The type of the media source being resolved. */
+	EWmfMediaResolveType Type;
+
+	/** The URL of the media source being resolved. */
+	FString Url;
+
+public:
+
+	/** Creates and initializes a new instance. */
+	FWmfMediaResolveState(EWmfMediaResolveType InType, const FString& InUrl)
+		: Type(InType)
+		, Url(InUrl)
+		, RefCount(0)
+	{ }
+
+public:
+
+	// IUnknown interface
+
+	STDMETHODIMP_(ULONG) AddRef()
+	{
+		return FPlatformAtomics::InterlockedIncrement(&RefCount);
+	}
+
+#if _MSC_VER == 1900
+#pragma warning(push)
+#pragma warning(disable:4838)
+#endif // _MSC_VER == 1900
+	STDMETHODIMP QueryInterface(REFIID RefID, void** Object)
+	{
+		static const QITAB QITab[] =
+		{
+			QITABENT(FWmfMediaResolveState, IUnknown),
+			{ 0 }
+		};
+
+		return QISearch( this, QITab, RefID, Object );
+	}
+#if _MSC_VER == 1900
+#pragma warning(pop)
+#endif // _MSC_VER == 1900
+
+	STDMETHODIMP_(ULONG) Release()
+	{
+		int32 CurrentRefCount = FPlatformAtomics::InterlockedDecrement(&RefCount);
+	
+		if (CurrentRefCount == 0)
+		{
+			delete this;
+		}
+
+		return CurrentRefCount;
+	}
+
+private:
+
+	/** Holds a reference counter for this instance. */
+	int32 RefCount;
+};
+
+
+#include "HideWindowsPlatformTypes.h"
diff --git a/Engine/Plugins/Media/WmfMedia/Source/WmfMedia/Private/Wmf/WmfMediaResolver.cpp b/Engine/Plugins/Media/WmfMedia/Source/WmfMedia/Private/Wmf/WmfMediaResolver.cpp
new file mode 100644
--- /dev/null
+++ b/Engine/Plugins/Media/WmfMedia/Source/WmfMedia/Private/Wmf/WmfMediaResolver.cpp
@@ -0,0 +1,171 @@
+// Copyright 1998-2015 Epic Games, Inc. All Rights Reserved.
+
+#include "WmfMediaPrivatePCH.h"
+#include "AllowWindowsPlatformTypes.h"
+
+
+/* WmfMediaUrlResolver structors
+ *****************************************************************************/
+
+FWmfMediaResolver::FWmfMediaResolver()
+	: RefCount(0)
+{
+	if (FAILED(::MFCreateSourceResolver(&SourceResolver)))
+	{
+		UE_LOG(LogWmfMedia, Error, TEXT("Failed to create media source resolver."));
+	}
+}
+
+
+/* WmfMediaUrlResolver interface
+ *****************************************************************************/
+
+void FWmfMediaResolver::Cancel()
+{
+	if (SourceResolver == NULL)
+	{
+		return;
+	}
+
+	TComPtr<IUnknown> CancelCookieCopy = CancelCookie;
+
+	if (CancelCookieCopy != NULL)
+	{
+		SourceResolver->CancelObjectCreation(CancelCookieCopy);
+		CancelCookie.Reset();
+	}
+
+	ResolveState.Reset();
+}
+
+
+bool FWmfMediaResolver::ResolveByteStream(const TSharedRef<FArchive, ESPMode::ThreadSafe>& Archive, const FString& OriginalUrl)
+{
+	if (SourceResolver == NULL)
+	{
+		return false;
+	}
+
+	Cancel();
+
+	TComPtr<FWmfMediaByteStream> ByteStream = new FWmfMediaByteStream(Archive);
+	ResolveState = new(std::nothrow) FWmfMediaResolveState(EWmfMediaResolveType::ByteStream, OriginalUrl);
+
+	return SUCCEEDED(SourceResolver->BeginCreateObjectFromByteStream(ByteStream, *OriginalUrl, MF_RESOLUTION_MEDIASOURCE, NULL, &CancelCookie, this, ResolveState));
+}
+
+
+bool FWmfMediaResolver::ResolveUrl(const FString& Url)
+{
+	if (SourceResolver == NULL)
+	{
+		return false;
+	}
+
+	Cancel();
+
+	ResolveState = new(std::nothrow) FWmfMediaResolveState(EWmfMediaResolveType::Url, Url);
+
+	return SUCCEEDED(SourceResolver->BeginCreateObjectFromURL(*Url, MF_RESOLUTION_MEDIASOURCE, NULL, &CancelCookie, this, ResolveState));
+}
+
+
+/* IMFAsyncCallback interface
+ *****************************************************************************/
+
+STDMETHODIMP_(ULONG) FWmfMediaResolver::AddRef()
+{
+	return FPlatformAtomics::InterlockedIncrement(&RefCount);
+}
+
+
+STDMETHODIMP FWmfMediaResolver::GetParameters(unsigned long*, unsigned long*)
+{
+	return E_NOTIMPL;
+}
+
+
+STDMETHODIMP FWmfMediaResolver::Invoke(IMFAsyncResult* AsyncResult)
+{
+	// are we currently resolving anything?
+	TComPtr<FWmfMediaResolveState> ResolveStateCopy = ResolveState;
+
+	if (ResolveStateCopy == NULL)
+	{
+		return S_OK;
+	}
+
+	// is this the resolve state we care about?
+	FWmfMediaResolveState* State = NULL;
+
+	if (FAILED(AsyncResult->GetState((IUnknown**)&State)) || (State != ResolveStateCopy))
+	{
+		return S_OK;
+	}
+
+	CancelCookie.Reset();
+
+	// get the asynchronous result
+	MF_OBJECT_TYPE ObjectType;
+	TComPtr<IUnknown> SourceObject;
+	bool Succeeded = false;
+
+	if (State->Type == EWmfMediaResolveType::ByteStream)
+	{
+		Succeeded = SUCCEEDED(SourceResolver->EndCreateObjectFromByteStream(AsyncResult, &ObjectType, &SourceObject));
+	}
+	else if (State->Type == EWmfMediaResolveType::Url)
+	{
+		Succeeded = SUCCEEDED(SourceResolver->EndCreateObjectFromURL(AsyncResult, &ObjectType, &SourceObject));
+	}
+
+	if (Succeeded)
+	{
+		ResolveCompletedDelegate.ExecuteIfBound(SourceObject, State->Url);
+	}
+	else
+	{
+		ResolveFailedDelegate.ExecuteIfBound(State->Url);
+	}
+
+	ResolveState.Reset();
+
+	return S_OK;
+}
+
+
+#if _MSC_VER == 1900
+#pragma warning(push)
+#pragma warning(disable:4838)
+#endif // _MSC_VER == 1900
+
+STDMETHODIMP FWmfMediaResolver::QueryInterface(REFIID RefID, void** Object)
+{
+	static const QITAB QITab[] =
+	{
+		QITABENT(FWmfMediaResolver, IMFAsyncCallback),
+		{ 0 }
+	};
+
+	return QISearch( this, QITab, RefID, Object );
+}
+
+#if _MSC_VER == 1900
+#pragma warning(pop)
+#endif // _MSC_VER == 1900
+
+
+STDMETHODIMP_(ULONG) FWmfMediaResolver::Release()
+{
+	int32 CurrentRefCount = FPlatformAtomics::InterlockedDecrement(&RefCount);
+	
+	if (CurrentRefCount == 0)
+	{
+		delete this;
+	}
+
+	return CurrentRefCount;
+}
+
+
+#include "HideWindowsPlatformTypes.h"
diff --git a/Engine/Plugins/Media/WmfMedia/Source/WmfMedia/Private/Wmf/WmfMediaResolver.h b/Engine/Plugins/Media/WmfMedia/Source/WmfMedia/Private/Wmf/WmfMediaResolver.h
new file mode 100644
--- /dev/null
+++ b/Engine/Plugins/Media/WmfMedia/Source/WmfMedia/Private/Wmf/WmfMediaResolver.h
@@ -0,0 +1,111 @@
+// Copyright 1998-2015 Epic Games, Inc. All Rights Reserved.
+
+#pragma once
+
+#include "AllowWindowsPlatformTypes.h"
+
+
+/** Delegate type that is executed when URL resolving has completed. */
+DECLARE_DELEGATE_TwoParams(FWmfResolveCompletedDelegate, TComPtr<IUnknown> /*SourceObject*/, FString /*ResolvedUrl*/);
+
+/** Delegate type that is executed when URL resolving has failed. */
+DECLARE_DELEGATE_OneParam(FWmfResolveFailedDelegate, FString /*FailedUrl*/);
+
+
+/**
+ * Implements an asynchronous callback object for resolving media URLs.
+ */
+class FWmfMediaResolver
+	: public IMFAsyncCallback
+{
+public:
+
+	/** Default constructor. */
+	FWmfMediaResolver();
+
+public:
+
+	/**
+	 * Cancel the current media source resolve.
+	 *
+	 * Note: Because resolving media sources is an asynchronous operation, the
+	 * callback delegates may still get invoked even after calling this method.
+	 *
+	 * @see ResolveByteStream, ResolveUrl
+	 */
+	void Cancel();
+
+	/**
+	 * Get a delegate that is executed when URL resolving has completed.
+	 *
+	 * @return The delegate.
+	 */
+	FWmfResolveCompletedDelegate& OnResolveComplete()
+	{
+		return ResolveCompletedDelegate;
+	}
+
+	/**
+	 * Get a delegate that is executed when URL resolving has failed.
+	 *
+	 * @return The delegate.
+	 */
+	FWmfResolveFailedDelegate& OnResolveFailed()
+	{
+		return ResolveFailedDelegate;
+	}
+
+	/**
+	 * Resolves a media source from a byte stream.
+	 *
+	 * @param Archive The archive holding the media data.
+	 * @param OriginalUrl The original URL of the media that was loaded into the buffer.
+	 * @return true if the byte stream will be resolved, false otherwise.
+	 * @see Cancel, ResolveUrl
+	 */
+	bool ResolveByteStream(const TSharedRef<FArchive, ESPMode::ThreadSafe>& Archive, const FString& OriginalUrl);
+
+	/**
+	 * Resolves a media source from a file or internet URL.
+	 *
+	 * @param Url The URL of the media to open (file name or web address).
+	 * @return true if the URL will be resolved, false otherwise.
+	 * @see Cancel, ResolveByteStream
+	 */
+	bool ResolveUrl(const FString& Url);
+
+public:
+
+	// IMFAsyncCallback interface
+
+	STDMETHODIMP_(ULONG) AddRef();
+	STDMETHODIMP GetParameters(unsigned long*, unsigned long*);
+	STDMETHODIMP Invoke(IMFAsyncResult* AsyncResult);
+	STDMETHODIMP QueryInterface(REFIID RefID, void** Object);
+	STDMETHODIMP_(ULONG) Release();
+
+private:
+
+	/** Cancellation cookie object. */
+	TComPtr<IUnknown> CancelCookie;
+
+	/** Holds a reference counter for this instance. */
+	int32 RefCount;
+
+	/** The current resolve state. */
+	TComPtr<FWmfMediaResolveState> ResolveState;
+
+	/** Holds the actual source resolver. */
+	TComPtr<IMFSourceResolver> SourceResolver;
+
+private:
+
+	/** A delegate that is executed when URL resolving has completed. */
+	FWmfResolveCompletedDelegate ResolveCompletedDelegate;
+
+	/** A delegate that is executed when URL resolving has failed. */
+	FWmfResolveFailedDelegate ResolveFailedDelegate;
+};
+
+
+#include "HideWindowsPlatformTypes.h"
diff --git a/Engine/Plugins/Media/WmfMedia/Source/WmfMedia/Private/WmfMediaPrivatePCH.h b/Engine/Plugins/Media/WmfMedia/Source/WmfMedia/Private/WmfMediaPrivatePCH.h
--- a/Engine/Plugins/Media/WmfMedia/Source/WmfMedia/Private/WmfMediaPrivatePCH.h
+++ b/Engine/Plugins/Media/WmfMedia/Source/WmfMedia/Private/WmfMediaPrivatePCH.h
@@ -42,4 +42,6 @@ DECLARE_LOG_CATEGORY_EXTERN(LogWmfMedia, Log, All);
 #include "WmfMediaReadState.h"
 #include "WmfMediaByteStream.h"
 #include "WmfMediaSession.h"
+#include "WmfMediaResolveState.h"
+#include "WmfMediaResolver.h"
 #include "WmfMediaPlayer.h"
diff --git a/Engine/Source/Runtime/Media/Public/IMediaPlayer.h b/Engine/Source/Runtime/Media/Public/IMediaPlayer.h
--- a/Engine/Source/Runtime/Media/Public/IMediaPlayer.h
+++ b/Engine/Source/Runtime/Media/Public/IMediaPlayer.h
@@ -140,21 +140,31 @@ public:
 	virtual bool IsReady() const = 0;
 
 	/**
-	 * Opens a media from a file on disk.
+	 * Opens a media from a URL, possibly asynchronously.
+	 *
+	 * The media may not necessarily be opened after this function succeeds,
+	 * because opening may happen asynchronously. Subscribe to the OnOpened
+	 * OnOpenFailed events to detect when the media finished or failed to
+	 * open. These events are only triggered if Open returns true.
 	 *
 	 * @param Url The URL of the media to open (file name or web address).
-	 * @return true if the media was opened successfully, false otherwise.
-	 * @see Close, IsReady
+	 * @return true if the media is being opened, false otherwise.
+	 * @see Close, IsReady, OnOpen, OnOpenFailed
 	 */
 	virtual bool Open(const FString& Url) = 0;
 
 	/**
-	 * Opens a media from a file or memory archive.
+	 * Opens a media from a file or memory archive, possibly asynchronously.
+	 *
+	 * The media may not necessarily be opened after this function succeeds,
+	 * because opening may happen asynchronously. Subscribe to the OnOpened
+	 * OnOpenFailed events to detect when the media finished or failed to
+	 * open. These events are only triggered if Open returns true.
 	 *
 	 * @param Archive The archive holding the media data.
 	 * @param OriginalUrl The original URL of the media that was loaded into the buffer.
-	 * @return true if the media was opened, false otherwise.
-	 * @see Close, IsReady
+	 * @return true if the media is being opened, false otherwise.
+	 * @see Close, IsReady, OnOpen, OnOpenFailed
 	 */
 	virtual bool Open(const TSharedRef<FArchive, ESPMode::ThreadSafe>& Archive, const FString& OriginalUrl) = 0;
 
@@ -199,6 +209,10 @@ public:
 	DECLARE_EVENT_OneParam(IMediaPlayer, FOnMediaOpened, FString /*OpenedUrl*/)
 	virtual FOnMediaOpened& OnOpened() = 0;
 
+	/** Gets an event delegate that is invoked when media failed to open. */
+	DECLARE_EVENT_OneParam(IMediaPlayer, FOnMediaOpenFailed, FString /*FailedUrl*/)
+	virtual FOnMediaOpenFailed& OnOpenFailed() = 0;
+
 	/** Gets an event delegate that is invoked when the media tracks have changed. */
 	DECLARE_EVENT(IMediaPlayer, FOnTracksChanged)
 	virtual FOnTracksChanged& OnTracksChanged() = 0;
diff --git a/Engine/Source/Runtime/MediaAssets/Private/Assets/MediaPlayer.cpp b/Engine/Source/Runtime/MediaAssets/Private/Assets/MediaPlayer.cpp
--- a/Engine/Source/Runtime/MediaAssets/Private/Assets/MediaPlayer.cpp
+++ b/Engine/Source/Runtime/MediaAssets/Private/Assets/MediaPlayer.cpp
@@ -234,6 +234,7 @@ void UMediaPlayer::InitializePlayer()
 			Player->Close();
 			Player->OnClosed().RemoveAll(this);
 			Player->OnOpened().RemoveAll(this);
+			Player->OnOpenFailed().RemoveAll(this);
 			Player->OnTracksChanged().RemoveAll(this);
 			Player.Reset();
 		}
@@ -260,6 +261,7 @@ void UMediaPlayer::InitializePlayer()
 
 		Player->OnClosed().AddUObject(this, &UMediaPlayer::HandlePlayerMediaClosed);
 		Player->OnOpened().AddUObject(this, &UMediaPlayer::HandlePlayerMediaOpened);
+		Player->OnOpenFailed().AddUObject(this, &UMediaPlayer::HandlePlayerMediaOpenFailed);
 		Player->OnTracksChanged().AddUObject(this, &UMediaPlayer::HandlePlayerTracksChanged);
 
 		// open the new media file
@@ -304,13 +306,19 @@ void UMediaPlayer::HandlePlayerMediaClosed()
 }
 
 
-void UMediaPlayer::HandlePlayerMediaOpened( FString OpenedUrl )
+void UMediaPlayer::HandlePlayerMediaOpened(FString OpenedUrl)
 {
 	MediaChangedEvent.Broadcast();
 	OnMediaOpened.Broadcast(OpenedUrl);
 }
 
 
+void UMediaPlayer::HandlePlayerMediaOpenFailed(FString FailedUrl)
+{
+	OnMediaOpenFailed.Broadcast(FailedUrl);
+}
+
+
 void UMediaPlayer::HandlePlayerTracksChanged()
 {
 	TracksChangedEvent.Broadcast();
diff --git a/Engine/Source/Runtime/MediaAssets/Public/MediaPlayer.h b/Engine/Source/Runtime/MediaAssets/Public/MediaPlayer.h
--- a/Engine/Source/Runtime/MediaAssets/Public/MediaPlayer.h
+++ b/Engine/Source/Runtime/MediaAssets/Public/MediaPlayer.h
@@ -14,6 +14,9 @@ DECLARE_DYNAMIC_MULTICAST_DELEGATE(FOnMediaPlayerMediaClosed);
 /** Multicast delegate that is invoked when a media player's media has been opened. */
 DECLARE_DYNAMIC_MULTICAST_DELEGATE_OneParam(FOnMediaPlayerMediaOpened, FString, OpenedUrl);
 
+/** Multicast delegate that is invoked when a media player's media has failed to open. */
+DECLARE_DYNAMIC_MULTICAST_DELEGATE_OneParam(FOnMediaPlayerMediaOpenFailed, FString, FailedUrl);
+
 
 /**
  * Implements a media player asset that can play movies and other media.
@@ -180,7 +183,7 @@ public:
 	 * @see GetTime, Rewind
 	 */
 	UFUNCTION(BlueprintCallable, Category="Media|MediaPlayer")
-	bool Seek( const FTimespan& InTime );
+	bool Seek(const FTimespan& InTime);
 
 	/**
 	 * Enables or disables playback looping.
@@ -190,7 +193,7 @@ public:
 	 * @see IsLooping
 	 */
 	UFUNCTION(BlueprintCallable, Category="Media|MediaPlayer")
-	bool SetLooping( bool InLooping );
+	bool SetLooping(bool InLooping);
 
 	/**
 	 * Changes the media's playback rate.
@@ -200,7 +203,7 @@ public:
 	 * @see GetRate, SupportsRate
 	 */
 	UFUNCTION(BlueprintCallable, Category="Media|MediaPlayer")
-	bool SetRate( float Rate );
+	bool SetRate(float Rate);
 
 	/**
 	 * Checks whether the specified playback rate is supported.
@@ -210,7 +213,7 @@ public:
 	 * @see SupportsScrubbing, SupportsSeeking
 	 */
 	UFUNCTION(BlueprintCallable, Category="Media|MediaPlayer")
-	bool SupportsRate( float Rate, bool Unthinned ) const;
+	bool SupportsRate(float Rate, bool Unthinned) const;
 
 	/**
 	 * Checks whether the currently loaded media supports scrubbing.
@@ -254,6 +257,10 @@ public:
 	UPROPERTY(BlueprintAssignable, Category="Media|MediaPlayer")
 	FOnMediaPlayerMediaOpened OnMediaOpened;
 
+	/** Holds a delegate that is invoked when a media source has failed to open. */
+	UPROPERTY(BlueprintAssignable, Category="Media|MediaPlayer")
+	FOnMediaPlayerMediaOpenFailed OnMediaOpenFailed;
+
 public:
 
 	/**
@@ -275,7 +282,7 @@ public:
 	virtual void PostLoad() override;
 
 #if WITH_EDITOR
-	virtual void PostEditChangeProperty( FPropertyChangedEvent& PropertyChangedEvent ) override;
+	virtual void PostEditChangeProperty(FPropertyChangedEvent& PropertyChangedEvent) override;
 #endif
 
 protected:
@@ -299,7 +306,10 @@ private:
 	void HandlePlayerMediaClosed();
 
 	/** Callback for when the player has opened a new media source. */
-	void HandlePlayerMediaOpened( FString OpenedUrl );
+	void HandlePlayerMediaOpened(FString OpenedUrl);
+
+	/** Callback for when the player has failed to open a new media source. */
+	void HandlePlayerMediaOpenFailed(FString FailedUrl);
 
 	/** Callback for when the player's tracks changed. */
 	void HandlePlayerTracksChanged();
