Copying //UE4/Dev-Sequencer to //UE4/Dev-Main (Source: //UE4/Dev-Sequencer @ 3057646)

diff --git a/Engine/Plugins/Media/WmfMedia/Source/WmfMedia/Private/Player/WmfMediaPlayer.cpp b/Engine/Plugins/Media/WmfMedia/Source/WmfMedia/Private/Player/WmfMediaPlayer.cpp
--- a/Engine/Plugins/Media/WmfMedia/Source/WmfMedia/Private/Player/WmfMediaPlayer.cpp
+++ b/Engine/Plugins/Media/WmfMedia/Source/WmfMedia/Private/Player/WmfMediaPlayer.cpp
@@ -1,6 +1,12 @@
 // Copyright 1998-2016 Epic Games, Inc. All Rights Reserved.
 
-#include "WmfMediaPrivatePCH.h"
+#include "WmfMediaPCH.h"
+#include "WmfMediaPlayer.h"
+#include "WmfMediaResolver.h"
+#include "WmfMediaSession.h"
+#include "WmfMediaSampler.h"
+#include "WmfMediaSettings.h"
+#include "WmfMediaUtils.h"
 #include "AllowWindowsPlatformTypes.h"
 
 
@@ -9,9 +15,9 @@
 
 FWmfMediaPlayer::FWmfMediaPlayer()
 	: Duration(0)
-	, MediaSession(nullptr)
 {
-	Resolver = new(std::nothrow) FWmfMediaResolver;
+	MediaSession = new FWmfMediaSession();
+	Resolver = new FWmfMediaResolver;
 }
 
 
@@ -28,7 +34,7 @@ bool FWmfMediaPlayer::Tick(float DeltaTime)
 {
 	TFunction<void()> Task;
 
-	while (Tasks.Dequeue(Task))
+	while (GameThreadTasks.Dequeue(Task))
 	{
 		Task();
 	}
@@ -37,54 +43,6 @@ bool FWmfMediaPlayer::Tick(float DeltaTime)
 }
 
 
-/* IMediaInfo interface
- *****************************************************************************/
-
-FTimespan FWmfMediaPlayer::GetDuration() const
-{
-	return Duration;
-}
-
-
-TRange<float> FWmfMediaPlayer::GetSupportedRates(EMediaPlaybackDirections Direction, bool Unthinned) const
-{
-	if (MediaSession != NULL)
-	{
-		return MediaSession->GetSupportedRates(Direction, Unthinned);
-	}
-
-	return TRange<float>(0.0f);
-}
-
-
-FString FWmfMediaPlayer::GetUrl() const
-{
-	return MediaUrl;
-}
-
-
-bool FWmfMediaPlayer::SupportsRate(float Rate, bool Unthinned) const
-{
-	return (MediaSession != NULL)
-		? MediaSession->IsRateSupported(Rate, Unthinned)
-		: false;
-}
-
-
-bool FWmfMediaPlayer::SupportsScrubbing() const
-{
-	return ((MediaSession != NULL) && MediaSession->SupportsScrubbing());
-}
-
-
-bool FWmfMediaPlayer::SupportsSeeking() const
-{
-	return ((MediaSession != NULL) &&
-			((MediaSession->GetCapabilities() & MFSESSIONCAP_SEEK) != 0) &&
-			(Duration > FTimespan::Zero()));
-}
-
-
 /* IMediaPlayer interface
  *****************************************************************************/
 
@@ -92,129 +50,133 @@ void FWmfMediaPlayer::Close()
 {
 	Resolver->Cancel();
 
-	if (MediaSession == NULL)
+	if ((MediaSession == NULL) || (MediaSession->GetState() == EMediaState::Closed))
 	{
 		return;
 	}
 
-	MediaSession->OnError().RemoveAll(this);
+	// unregister session callbacks
 	MediaSession->OnSessionEvent().RemoveAll(this);
 
-	if (IsPlaying())
+	if (MediaSession->GetState() == EMediaState::Playing)
 	{
 		MediaEvent.Broadcast(EMediaEvent::PlaybackSuspended);
 	}
 
-	MediaSession->SetState(EMediaStates::Closed);
-	MediaSession.Reset();
+	// close and release session
+	MediaSession->SetState(EMediaState::Closed);
+	MediaSession = new FWmfMediaSession();
 
+	// release media source
 	if (MediaSource != NULL)
 	{
 		MediaSource->Shutdown();
 		MediaSource = NULL;
 	}
 
-	AudioTracks.Reset();
-	CaptionTracks.Reset();
-	VideoTracks.Reset();
-
+	// reset player
 	Duration = 0;
+	Info.Empty();
 	MediaUrl = FString();
+	Tracks.Reset();
 
+	// notify listeners
 	MediaEvent.Broadcast(EMediaEvent::TracksChanged);
 	MediaEvent.Broadcast(EMediaEvent::MediaClosed);
 }
 
 
-const TArray<IMediaAudioTrackRef>& FWmfMediaPlayer::GetAudioTracks() const
+IMediaControls& FWmfMediaPlayer::GetControls()
 {
-	return AudioTracks;
+	return *MediaSession;
 }
 
 
-const TArray<IMediaCaptionTrackRef>& FWmfMediaPlayer::GetCaptionTracks() const
+FString FWmfMediaPlayer::GetInfo() const
 {
-	return CaptionTracks;
+	return Info;
 }
 
 
-const IMediaInfo& FWmfMediaPlayer::GetMediaInfo() const 
+IMediaOutput& FWmfMediaPlayer::GetOutput()
 {
-	return *this;
+	return Tracks;
 }
 
 
-float FWmfMediaPlayer::GetRate() const
+FString FWmfMediaPlayer::GetStats() const
 {
-	return (MediaSession != NULL)
-		? MediaSession->GetRate()
-		: 0.0f;
-}
+	FString Result;
 
+	Tracks.AppendStats(Result);
 
-FTimespan FWmfMediaPlayer::GetTime() const 
-{
-	return (MediaSession != NULL)
-		? MediaSession->GetPosition()
-		: FTimespan::Zero();
+	return Result;
 }
 
 
-const TArray<IMediaVideoTrackRef>& FWmfMediaPlayer::GetVideoTracks() const
+IMediaTracks& FWmfMediaPlayer::GetTracks()
 {
-	return VideoTracks;
+	return Tracks;
 }
 
 
-bool FWmfMediaPlayer::IsLooping() const 
+FString FWmfMediaPlayer::GetUrl() const
 {
-	return ((MediaSession != NULL) && MediaSession->IsLooping());
+	return MediaUrl;
 }
 
 
-bool FWmfMediaPlayer::IsPaused() const
+bool FWmfMediaPlayer::Open(const FString& Url, const IMediaOptions& Options)
 {
-	if (MediaSession == NULL)
+	Close();
+
+	if (Url.IsEmpty())
 	{
 		return false;
 	}
 
-	if (MediaSession->GetState() == EMediaStates::Playing)
+	// open local files via platform file system
+	if (Url.StartsWith(TEXT("file://")))
 	{
-		return (MediaSession->GetRate() == 0.0);
-	}
-
-	return (MediaSession->GetState() == EMediaStates::Paused);
-}
+		TSharedPtr<FArchive, ESPMode::ThreadSafe> Archive;
+		const TCHAR* FilePath = &Url[7];
 
+		if (Options.GetMediaOption("PrecacheFile", false))
+		{
+			FBufferArchive* Buffer = new FBufferArchive;
 
-bool FWmfMediaPlayer::IsPlaying() const
-{
-	return (MediaSession != NULL) && (MediaSession->GetState() == EMediaStates::Playing) && !FMath::IsNearlyZero(MediaSession->GetRate());
-}
-
+			if (FFileHelper::LoadFileToArray(*Buffer, FilePath))
+			{
+				Archive = MakeShareable(Buffer);
+			}
+			else
+			{
+				delete Buffer;
+			}
+		}
+		else
+		{
+			Archive = MakeShareable(IFileManager::Get().CreateFileReader(FilePath));
+		}
 
-bool FWmfMediaPlayer::IsReady() const
-{
-	return ((MediaSession != NULL) &&
-			(MediaSession->GetState() != EMediaStates::Closed) &&
-			(MediaSession->GetState() != EMediaStates::Error));
-}
+		if (!Archive.IsValid())
+		{
+			UE_LOG(LogWmfMedia, Warning, TEXT("Failed to open media file: %s"), FilePath);
 
+			return false;
+		}
 
-bool FWmfMediaPlayer::Open(const FString& Url)
-{
-	if (Url.IsEmpty())
-	{
-		return false;
+		return Resolver->ResolveByteStream(Archive.ToSharedRef(), Url, *this);
 	}
 
 	return Resolver->ResolveUrl(Url, *this);
 }
 
 
-bool FWmfMediaPlayer::Open(const TSharedRef<FArchive, ESPMode::ThreadSafe>& Archive, const FString& OriginalUrl)
+bool FWmfMediaPlayer::Open(const TSharedRef<FArchive, ESPMode::ThreadSafe>& Archive, const FString& OriginalUrl, const IMediaOptions& Options)
 {
+	Close();
+
 	if ((Archive->TotalSize() == 0) || OriginalUrl.IsEmpty())
 	{
 		return false;
@@ -224,42 +186,12 @@ bool FWmfMediaPlayer::Open(const TSharedRef<FArchive, ESPMode::ThreadSafe>& Arch
 }
 
 
-bool FWmfMediaPlayer::Seek(const FTimespan& Time)
-{
-	return ((MediaSession != NULL) && MediaSession->SetPosition(Time));
-}
-
-
-bool FWmfMediaPlayer::SetLooping(bool Looping)
-{
-	if (MediaSession == NULL)
-	{
-		return false;
-	}
-
-	MediaSession->SetLooping(Looping);
-
-	return true;
-}
-
-
-bool FWmfMediaPlayer::SetRate(float Rate)
-{
-	if (MediaSession == NULL)
-	{
-		return false;
-	}
-
-	return MediaSession->SetRate(Rate);
-}
-
-
 /* IWmfMediaResolverCallbacks interface
  *****************************************************************************/
 
 void FWmfMediaPlayer::ProcessResolveComplete(TComPtr<IUnknown> SourceObject, FString ResolvedUrl)
 {
-	Tasks.Enqueue([=]() {
+	GameThreadTasks.Enqueue([=]() {
 		MediaEvent.Broadcast(
 			InitializeMediaSession(SourceObject, ResolvedUrl)
 				? EMediaEvent::MediaOpened
@@ -271,7 +203,7 @@ void FWmfMediaPlayer::ProcessResolveComplete(TComPtr<IUnknown> SourceObject, FSt
 
 void FWmfMediaPlayer::ProcessResolveFailed(FString FailedUrl)
 {
-	Tasks.Enqueue([=]() {
+	GameThreadTasks.Enqueue([=]() {
 		MediaEvent.Broadcast(EMediaEvent::MediaOpenFailed);
 	});
 }
@@ -282,147 +214,372 @@ void FWmfMediaPlayer::ProcessResolveFailed(FString FailedUrl)
 
 void FWmfMediaPlayer::AddStreamToTopology(uint32 StreamIndex, IMFTopology* Topology, IMFPresentationDescriptor* PresentationDescriptor, IMFMediaSource* MediaSourceObject)
 {
-	// get stream descriptor
-	TComPtr<IMFStreamDescriptor> StreamDescriptor;
-	BOOL Selected = FALSE;
+	Info += FString::Printf(TEXT("Stream %i\n"), StreamIndex);
 
-	if (FAILED(PresentationDescriptor->GetStreamDescriptorByIndex(StreamIndex, &Selected, &StreamDescriptor)))
+	// get stream descriptor & media type handler
+	TComPtr<IMFStreamDescriptor> StreamDescriptor;
 	{
-		UE_LOG(LogWmfMedia, Warning, TEXT("Skipping missing stream descriptor for stream index %i"), StreamIndex);
+		BOOL Selected = FALSE;
+		HRESULT Result = PresentationDescriptor->GetStreamDescriptorByIndex(StreamIndex, &Selected, &StreamDescriptor);
 
-		return;
+		if (FAILED(Result))
+		{
+			UE_LOG(LogWmfMedia, Warning, TEXT("Skipping missing stream descriptor for stream %i (%s)"), StreamIndex, *WmfMedia::ResultToString(Result));
+			Info += TEXT("    missing stream descriptor\n");
+			
+			return;
+		}
+
+		if (Selected == TRUE)
+		{
+			PresentationDescriptor->DeselectStream(StreamIndex);
+		}
 	}
 
-	// fetch media type
 	TComPtr<IMFMediaTypeHandler> Handler;
-
-	if (FAILED(StreamDescriptor->GetMediaTypeHandler(&Handler)))
 	{
-		return;
+		HRESULT Result = StreamDescriptor->GetMediaTypeHandler(&Handler);
+
+		if (FAILED(Result))
+		{
+			UE_LOG(LogWmfMedia, Warning, TEXT("Failed to get media type handler for stream %i (%s)"), StreamIndex, *WmfMedia::ResultToString(Result));
+			Info += TEXT("    no handler available\n");
+			
+			return;
+		}
 	}
 
+	// skip unsupported handler types
 	GUID MajorType;
-		
-	if (FAILED(Handler->GetMajorType(&MajorType)))
 	{
-		return;
-	}
+		HRESULT Result = Handler->GetMajorType(&MajorType);
 
-	// skip unsupported types
-	if ((MajorType != MFMediaType_Audio) &&
-		(MajorType != MFMediaType_SAMI) &&
-		(MajorType != MFMediaType_Video))
-	{
-		return;
-	}
+		if (FAILED(Result))
+		{
+			UE_LOG(LogWmfMedia, Warning, TEXT("Failed to determine major type of stream %i (%s)"), StreamIndex, *WmfMedia::ResultToString(Result));
+			Info += TEXT("    failed to determine MajorType\n");
+			
+			return;
+		}
 
-	// prepare input and output types
-	TComPtr<IMFMediaType> OutputType;
-	{
-		if (FAILED(Handler->GetCurrentMediaType(&OutputType)))
+		UE_LOG(LogWmfMedia, Verbose, TEXT("Major type of stream %i is %s"), StreamIndex, *WmfMedia::MajorTypeToString(MajorType));
+		Info += FString::Printf(TEXT("    Type: %s\n"), *WmfMedia::MajorTypeToString(MajorType));
+
+		if ((MajorType != MFMediaType_Audio) &&
+			(MajorType != MFMediaType_SAMI) &&
+			(MajorType != MFMediaType_Video))
 		{
+			UE_LOG(LogWmfMedia, Warning, TEXT("Unsupported major type %s for stream %i"), *WmfMedia::MajorTypeToString(MajorType), StreamIndex);
+			Info += TEXT("    MajorType is not supported\n");
+			
 			return;
 		}
 	}
 
-	TComPtr<IMFMediaType> InputType;
+	// get media type & make it current
+	TComPtr<IMFMediaType> MediaType;
 	{
-		if (FAILED(::MFCreateMediaType(&InputType)))
+		DWORD NumSupportedTypes = 0;
+		HRESULT Result = Handler->GetMediaTypeCount(&NumSupportedTypes);
+
+		if (FAILED(Result))
 		{
+			UE_LOG(LogWmfMedia, Warning, TEXT("Failed to get number of supported media types in stream %i of type %s"), StreamIndex, *WmfMedia::MajorTypeToString(MajorType));
+			Info += TEXT("    failed to get supported media types\n");
+
 			return;
 		}
 
-		if (FAILED(InputType->SetUINT32(MF_MT_ALL_SAMPLES_INDEPENDENT, TRUE)))
+		if (NumSupportedTypes > 0)
 		{
-			return;
+			for (DWORD TypeIndex = 0; TypeIndex < NumSupportedTypes; ++TypeIndex)
+			{
+				if (SUCCEEDED(Handler->GetMediaTypeByIndex(TypeIndex, &MediaType)) &&
+					SUCCEEDED(Handler->SetCurrentMediaType(MediaType)))
+				{
+					break;
+				}
+			}
 		}
-	}
 
-	if (MajorType == MFMediaType_Audio)
-	{
-		if (FAILED(InputType->SetGUID(MF_MT_MAJOR_TYPE, MFMediaType_Audio)) ||
-			FAILED(InputType->SetGUID(MF_MT_SUBTYPE, MFAudioFormat_PCM)))/* ||
-			FAILED(InputType->SetUINT32(MF_MT_AUDIO_SAMPLES_PER_BLOCK, 16384u)))*/
+		if (MediaType == NULL)
 		{
+			UE_LOG(LogWmfMedia, Warning, TEXT("No supported media type in stream %i of type %s"), StreamIndex, *WmfMedia::MajorTypeToString(MajorType));
+			Info += TEXT("    unsupported media type\n");
+
 			return;
 		}
 	}
-	else if (MajorType == MFMediaType_SAMI)
+	
+	// get sub-type
+	GUID SubType;
 	{
-		if (FAILED(InputType->SetGUID(MF_MT_MAJOR_TYPE, MFMediaType_SAMI)))
+		if (MajorType == MFMediaType_SAMI)
 		{
-			return;
+			FMemory::Memzero(SubType);
+		}
+		else
+		{
+			HRESULT Result = MediaType->GetGUID(MF_MT_SUBTYPE, &SubType);
+
+			if (FAILED(Result))
+			{
+				UE_LOG(LogWmfMedia, Warning, TEXT("Failed to get sub-type of stream %i (%s)"), StreamIndex, *WmfMedia::ResultToString(Result));
+				Info += TEXT("    failed to get sub-type\n");
+
+				return;
+			}
+
+			UE_LOG(LogWmfMedia, Verbose, TEXT("Sub-type of stream %i is %s"), StreamIndex, *WmfMedia::SubTypeToString(SubType));
+		
+			Info += FString::Printf(TEXT("    Codec: %s\n"), *WmfMedia::SubTypeToString(SubType));
+			Info += FString::Printf(TEXT("    Protected: %s\n"), (::MFGetAttributeUINT32(StreamDescriptor, MF_SD_PROTECTED, FALSE) != 0) ? *GYes.ToString() : *GNo.ToString());
 		}
 	}
-	else if (MajorType == MFMediaType_Video)
+
+	// configure desired output type
+	TComPtr<IMFMediaType> OutputType;
 	{
-		if (FAILED(InputType->SetGUID(MF_MT_MAJOR_TYPE, MFMediaType_Video)) ||
-			FAILED(InputType->SetGUID(MF_MT_SUBTYPE, MFVideoFormat_RGB32)))
+		HRESULT Result = ::MFCreateMediaType(&OutputType);
+
+		if (FAILED(Result))
 		{
+			UE_LOG(LogWmfMedia, Warning, TEXT("Failed to create output type for stream %i (%s)"), StreamIndex, *WmfMedia::ResultToString(Result));
+			Info += TEXT("    failed to create output type\n");
+
 			return;
 		}
-	}
 
-	// create sample grabber callback and add to topology
-	TComPtr<FWmfMediaSampler> Sampler = new FWmfMediaSampler();
-	{
-		TComPtr<IMFActivate> MediaSamplerActivator;
+		Result = OutputType->SetUINT32(MF_MT_ALL_SAMPLES_INDEPENDENT, TRUE);
 
-		if (FAILED(::MFCreateSampleGrabberSinkActivate(InputType, Sampler, &MediaSamplerActivator)))
+		if (FAILED(Result))
 		{
+			UE_LOG(LogWmfMedia, Warning, TEXT("Failed to initialize output type for stream %i (%s)"), StreamIndex, *WmfMedia::ResultToString(Result));
+			Info += TEXT("    failed to initialize output type\n");
+
 			return;
 		}
 
-		TComPtr<IMFTopologyNode> SourceNode;
+		if (MajorType == MFMediaType_Audio)
 		{
-			if (FAILED(::MFCreateTopologyNode(MF_TOPOLOGY_SOURCESTREAM_NODE, &SourceNode)) ||
-				FAILED(SourceNode->SetUnknown(MF_TOPONODE_SOURCE, MediaSourceObject)) ||
-				FAILED(SourceNode->SetUnknown(MF_TOPONODE_PRESENTATION_DESCRIPTOR, PresentationDescriptor)) ||
-				FAILED(SourceNode->SetUnknown(MF_TOPONODE_STREAM_DESCRIPTOR, StreamDescriptor)) ||
-				FAILED(Topology->AddNode(SourceNode)))
+			// get source stream information
+			Info += FString::Printf(TEXT("    Channels: %i\n"), ::MFGetAttributeUINT32(MediaType, MF_MT_AUDIO_NUM_CHANNELS, 0));
+			Info += FString::Printf(TEXT("    Sample Rate: %i Hz\n"), ::MFGetAttributeUINT32(MediaType, MF_MT_AUDIO_SAMPLES_PER_SECOND, 0));
+			Info += FString::Printf(TEXT("    Bits Per Sample: %i\n"), ::MFGetAttributeUINT32(MediaType, MF_MT_AUDIO_BITS_PER_SAMPLE, 0));
+
+			// filter unsupported audio formats
+			if (FMemory::Memcmp(&SubType.Data2, &MFMPEG4Format_Base.Data2, 12) == 0)
+			{
+				if (GetDefault<UWmfMediaSettings>()->AllowNonStandardCodecs)
+				{
+					UE_LOG(LogWmfMedia, Verbose, TEXT("Non-standard MP4 audio type '%s' (%s) for stream %i"), *WmfMedia::FourccToString(SubType.Data1), *WmfMedia::GuidToString(SubType), StreamIndex);
+					Info += TEXT("    non-standard sub-type\n");
+				}
+				else
+				{
+					const bool DocumentedFormat =
+						(SubType.Data1 == WAVE_FORMAT_ADPCM) ||
+						(SubType.Data1 == WAVE_FORMAT_ALAW) ||
+						(SubType.Data1 == WAVE_FORMAT_MULAW) ||
+						(SubType.Data1 == WAVE_FORMAT_IMA_ADPCM) ||
+						(SubType.Data1 == MFAudioFormat_AAC.Data1) ||
+						(SubType.Data1 == MFAudioFormat_MP3.Data1) ||
+						(SubType.Data1 == MFAudioFormat_PCM.Data1);
+
+					const bool UndocumentedFormat =
+						(SubType.Data1 == WAVE_FORMAT_WMAUDIO2) ||
+						(SubType.Data1 == WAVE_FORMAT_WMAUDIO3) ||
+						(SubType.Data1 == WAVE_FORMAT_WMAUDIO_LOSSLESS);
+
+					if (!DocumentedFormat && !UndocumentedFormat)
+					{
+						UE_LOG(LogWmfMedia, Warning, TEXT("Unsupported MP4 audio type '%s' (%s) for stream %i"), *WmfMedia::FourccToString(SubType.Data1), *WmfMedia::GuidToString(SubType), StreamIndex);
+						Info += TEXT("    unsupported sub-type\n");
+
+						return;
+					}
+				}
+			}
+			else if (FMemory::Memcmp(&SubType.Data2, &MFAudioFormat_Base.Data2, 12) != 0)
 			{
+				if (GetDefault<UWmfMediaSettings>()->AllowNonStandardCodecs)
+				{
+					UE_LOG(LogWmfMedia, Verbose, TEXT("Non-standard audio type '%s' (%s) for stream %i"), *WmfMedia::FourccToString(SubType.Data1), *WmfMedia::GuidToString(SubType), StreamIndex);
+					Info += TEXT("    non-standard sub-type\n");
+				}
+				else
+				{
+					UE_LOG(LogWmfMedia, Warning, TEXT("Unsupported audio type '%s' (%s) for stream %i"), *WmfMedia::FourccToString(SubType.Data1), *WmfMedia::GuidToString(SubType), StreamIndex);
+					Info += TEXT("    unsupported sub-type\n");
+
+					return;
+				}
+			}
+
+			// configure audio output (re-sampling fails for many media types, so we don't attempt it)
+			if (FAILED(OutputType->SetGUID(MF_MT_MAJOR_TYPE, MFMediaType_Audio)) ||
+				FAILED(OutputType->SetGUID(MF_MT_SUBTYPE, MFAudioFormat_PCM)) ||
+				FAILED(OutputType->SetUINT32(MF_MT_AUDIO_BITS_PER_SAMPLE, 16u)))
+			{
+				UE_LOG(LogWmfMedia, Warning, TEXT("Failed to initialize audio output type for stream %i"), StreamIndex);
+				Info += TEXT("    failed to initialize output type\n");
+
 				return;
 			}
+
+			if (::MFGetAttributeUINT32(MediaType, MF_MT_AUDIO_SAMPLES_PER_SECOND, 0) != 44100)
+			{
+				UE_LOG(LogWmfMedia, Warning, TEXT("Possible loss of audio quality in stream %i due to sample rate != 44100 Hz"), StreamIndex);
+			}
 		}
+		else if (MajorType == MFMediaType_SAMI)
+		{
+			// configure caption output
+			Result = OutputType->SetGUID(MF_MT_MAJOR_TYPE, MFMediaType_SAMI);
 
-		TComPtr<IMFTopologyNode> OutputNode;
+			if (FAILED(Result))
+			{
+				UE_LOG(LogWmfMedia, Warning, TEXT("Failed to initialize caption output type for stream %i (%s)"), StreamIndex, *WmfMedia::ResultToString(Result));
+				Info += TEXT("    failed to set output type\n");
+
+				return;
+			}
+		}
+		else if (MajorType == MFMediaType_Video)
 		{
-			if (FAILED(::MFCreateTopologyNode(MF_TOPOLOGY_OUTPUT_NODE, &OutputNode)) ||
-				FAILED(OutputNode->SetObject(MediaSamplerActivator)) ||
-				FAILED(OutputNode->SetUINT32(MF_TOPONODE_STREAMID, 0)) ||
-				FAILED(OutputNode->SetUINT32(MF_TOPONODE_NOSHUTDOWN_ON_REMOVE, FALSE)) ||
-				FAILED(Topology->AddNode(OutputNode)))
+			// get source stream information
+			FIntPoint Dimensions;
+
+			if (SUCCEEDED(::MFGetAttributeSize(MediaType, MF_MT_FRAME_SIZE, (UINT32*)&Dimensions.X, (UINT32*)&Dimensions.Y)))
+			{
+				Info += FString::Printf(TEXT("    Dimensions: %i x %i\n"), Dimensions.X, Dimensions.Y);
+			}
+			else
+			{
+				Info += FString::Printf(TEXT("    Dimensions: n/a\n"));
+			}
+
+			UINT32 Numerator = 0;
+			UINT32 Denominator = 1;
+
+			if (SUCCEEDED(::MFGetAttributeRatio(MediaType, MF_MT_FRAME_RATE, &Numerator, &Denominator)))
+			{
+				const float FrameRate = static_cast<float>(Numerator) / Denominator;
+				Info += FString::Printf(TEXT("    Frame Rate: %g fps\n"), FrameRate);
+			}
+			else
+			{
+				Info += FString::Printf(TEXT("    Frame Rate: n/a\n"));
+			}
+
+			// filter unsupported video types
+			if (FMemory::Memcmp(&SubType.Data2, &MFVideoFormat_Base.Data2, 12) != 0)
+			{
+				if (GetDefault<UWmfMediaSettings>()->AllowNonStandardCodecs)
+				{
+					UE_LOG(LogWmfMedia, Verbose, TEXT("Non-standard video type '%s' (%s) for stream %i"), *WmfMedia::FourccToString(SubType.Data1), *WmfMedia::GuidToString(SubType), StreamIndex);
+					Info += TEXT("    non-standard sub-type\n");
+				}
+				else
+				{
+					UE_LOG(LogWmfMedia, Warning, TEXT("Unsupported video type '%s' (%s) for stream %i"), *WmfMedia::FourccToString(SubType.Data1), *WmfMedia::GuidToString(SubType), StreamIndex);
+					Info += TEXT("    unsupported sub-type\n");
+				}
+
+				return;
+			}
+
+			// configure video output
+			Result = OutputType->SetGUID(MF_MT_MAJOR_TYPE, MFMediaType_Video);
+
+			if (FAILED(Result))
+			{
+				UE_LOG(LogWmfMedia, Warning, TEXT("Failed to set video output type for stream %i (%s)"), StreamIndex, *WmfMedia::ResultToString(Result));
+				Info += TEXT("    failed to set output type\n");
+
+				return;
+			}
+
+			GUID OutputSubType = ((SubType == MFVideoFormat_H264) || (SubType == MFVideoFormat_H264_ES))
+				? MFVideoFormat_YUY2
+				: MFVideoFormat_RGB32;
+
+			Result = OutputType->SetGUID(MF_MT_SUBTYPE, OutputSubType);
+
+			if (FAILED(Result))
 			{
+				UE_LOG(LogWmfMedia, Warning, TEXT("Failed to set video output sub-type for stream %i (%s)"), StreamIndex, *WmfMedia::ResultToString(Result));
+				Info += TEXT("    failed to set output sub-type\n");
+
 				return;
 			}
 		}
+	}
+
+	TComPtr<FWmfMediaSampler> Sampler = new FWmfMediaSampler();
 
-		if (FAILED(SourceNode->ConnectOutput(0, OutputNode, 0)))
+	// set up output node
+	TComPtr<IMFActivate> MediaSamplerActivator;
+	{
+		HRESULT Result = ::MFCreateSampleGrabberSinkActivate(OutputType, Sampler, &MediaSamplerActivator);
+
+		if (FAILED(Result))
 		{
+			UE_LOG(LogWmfMedia, Warning, TEXT("Failed to create sampler grabber sink for stream %i (%s)"), StreamIndex, *WmfMedia::ResultToString(Result));
+			Info += TEXT("    failed to create sample grabber\n");
+
 			return;
 		}
 	}
 
-	// create and add track
-	if (MajorType == MFMediaType_Audio)
+	TComPtr<IMFTopologyNode> OutputNode;
 	{
-		AudioTracks.Add(MakeShareable(new FWmfMediaAudioTrack(OutputType, PresentationDescriptor, Sampler, StreamDescriptor, StreamIndex)));
+		if (FAILED(::MFCreateTopologyNode(MF_TOPOLOGY_OUTPUT_NODE, &OutputNode)) ||
+			FAILED(OutputNode->SetObject(MediaSamplerActivator)) ||
+			FAILED(OutputNode->SetUINT32(MF_TOPONODE_STREAMID, 0)) ||
+			FAILED(OutputNode->SetUINT32(MF_TOPONODE_NOSHUTDOWN_ON_REMOVE, FALSE)) ||
+			FAILED(Topology->AddNode(OutputNode)))
+		{
+			UE_LOG(LogWmfMedia, Warning, TEXT("Failed to configure output node for stream %i"), StreamIndex);
+			Info += TEXT("    failed to configure output node\n");
+
+			return;
+		}
 	}
-	else if (MajorType == MFMediaType_SAMI)
+
+	// set up source node
+	TComPtr<IMFTopologyNode> SourceNode;
 	{
-		CaptionTracks.Add(MakeShareable(new FWmfMediaCaptionTrack(OutputType, PresentationDescriptor, Sampler, StreamDescriptor, StreamIndex)));
+		if (FAILED(::MFCreateTopologyNode(MF_TOPOLOGY_SOURCESTREAM_NODE, &SourceNode)) ||
+			FAILED(SourceNode->SetUnknown(MF_TOPONODE_SOURCE, MediaSourceObject)) ||
+			FAILED(SourceNode->SetUnknown(MF_TOPONODE_PRESENTATION_DESCRIPTOR, PresentationDescriptor)) ||
+			FAILED(SourceNode->SetUnknown(MF_TOPONODE_STREAM_DESCRIPTOR, StreamDescriptor)) ||
+			FAILED(Topology->AddNode(SourceNode)))
+		{
+			UE_LOG(LogWmfMedia, Warning, TEXT("Failed to configure source node for stream %i"), StreamIndex);
+			Info += TEXT("    failed to configure source node\n");
+
+			return;
+		}
 	}
-	else if (MajorType == MFMediaType_Video)
+
+	// connect nodes
+	HRESULT Result = SourceNode->ConnectOutput(0, OutputNode, 0);
+
+	if (FAILED(Result))
 	{
-		VideoTracks.Add(MakeShareable(new FWmfMediaVideoTrack(OutputType, PresentationDescriptor, Sampler, StreamDescriptor, StreamIndex)));
+		UE_LOG(LogWmfMedia, Warning, TEXT("Failed to connect topology nodes for stream %i (%s)"), StreamIndex, *WmfMedia::ResultToString(Result));
+		Info += TEXT("    failed to connect topology nodes\n");
+
+		return;
 	}
+
+	Tracks.AddTrack(MajorType, MediaType, OutputType, Sampler, StreamDescriptor, StreamIndex);
 }
 
 
 bool FWmfMediaPlayer::InitializeMediaSession(IUnknown* SourceObject, const FString& SourceUrl)
 {
-	Close();
-
 	if (SourceObject == nullptr)
 	{
 		return false;
@@ -432,50 +589,64 @@ bool FWmfMediaPlayer::InitializeMediaSession(IUnknown* SourceObject, const FStri
 
 	// create presentation descriptor
 	TComPtr<IMFMediaSource> MediaSourceObject;
-
-	if (FAILED(SourceObject->QueryInterface(IID_PPV_ARGS(&MediaSourceObject))))
 	{
-		UE_LOG(LogWmfMedia, Error, TEXT("Failed to query media source"));
+		HRESULT Result = SourceObject->QueryInterface(IID_PPV_ARGS(&MediaSourceObject));
 
-		return false;
+		if (FAILED(Result))
+		{
+			UE_LOG(LogWmfMedia, Error, TEXT("Failed to query media source: %s"), *WmfMedia::ResultToString(Result));
+			return false;
+		}
 	}
 
 	TComPtr<IMFPresentationDescriptor> PresentationDescriptor;
-	
-	if (FAILED(MediaSourceObject->CreatePresentationDescriptor(&PresentationDescriptor)))
 	{
-		UE_LOG(LogWmfMedia, Error, TEXT("Failed to create presentation descriptor"));
+		HRESULT Result = MediaSourceObject->CreatePresentationDescriptor(&PresentationDescriptor);
 
-		return false;
+		if (FAILED(Result))
+		{
+			UE_LOG(LogWmfMedia, Error, TEXT("Failed to create presentation descriptor: %s"), *WmfMedia::ResultToString(Result));
+			return false;
+		}
 	}
 	
 	// create playback topology
-	DWORD StreamCount = 0;
-
-	if (FAILED(PresentationDescriptor->GetStreamDescriptorCount(&StreamCount)))
+	TComPtr<IMFTopology> Topology;
 	{
-		UE_LOG(LogWmfMedia, Error, TEXT("Failed to get stream count"));
+		HRESULT Result = ::MFCreateTopology(&Topology);
 
-		return false;
+		if (FAILED(Result))
+		{
+			UE_LOG(LogWmfMedia, Error, TEXT("Failed to create playback topology: %s"), *WmfMedia::ResultToString(Result));
+			return false;
+		}
 	}
 
-	TComPtr<IMFTopology> Topology;
-
-	if (FAILED(::MFCreateTopology(&Topology)))
+	// initialize tracks
+	DWORD StreamCount = 0;
 	{
-		UE_LOG(LogWmfMedia, Error, TEXT("Failed to create playback topology"));
+		HRESULT Result = PresentationDescriptor->GetStreamDescriptorCount(&StreamCount);
 
-		return false;
+		if (FAILED(Result))
+		{
+			UE_LOG(LogWmfMedia, Error, TEXT("Failed to get stream count (%s)"), *WmfMedia::ResultToString(Result));
+			return false;
+		}
+
+		UE_LOG(LogWmfMedia, Verbose, TEXT("Found %i streams"), StreamCount);
 	}
 
+	Tracks.Initialize(*PresentationDescriptor);
+
 	for (uint32 StreamIndex = 0; StreamIndex < StreamCount; ++StreamIndex)
 	{
 		AddStreamToTopology(StreamIndex, Topology, PresentationDescriptor, MediaSourceObject);
+		Info += TEXT("\n");
 	}
 
-	UE_LOG(LogWmfMedia, Verbose, TEXT("Added a total of %i audio tracks, %i caption tracks, %i video tracks"), AudioTracks.Num(), CaptionTracks.Num(), VideoTracks.Num());
 	MediaEvent.Broadcast(EMediaEvent::TracksChanged);
 
+	// get media duration
 	UINT64 PresentationDuration = 0;
 	PresentationDescriptor->GetUINT64(MF_PD_DURATION, &PresentationDuration);
 	Duration = FTimespan(PresentationDuration);
@@ -484,24 +655,21 @@ bool FWmfMediaPlayer::InitializeMediaSession(IUnknown* SourceObject, const FStri
 	MediaUrl = SourceUrl;
 	MediaSource = MediaSourceObject;
 	MediaSession = new FWmfMediaSession(Duration, Topology);
+
+	if (MediaSession->GetState() == EMediaState::Error)
 	{
-		MediaSession->OnError().AddRaw(this, &FWmfMediaPlayer::HandleSessionError);
-		MediaSession->OnSessionEvent().AddRaw(this, &FWmfMediaPlayer::HandleSessionEvent);
+		return false;
 	}
 
-	return (MediaSession->GetState() != EMediaStates::Error);
+	MediaSession->OnSessionEvent().AddRaw(this, &FWmfMediaPlayer::HandleSessionEvent);
+
+	return true;
 }
 
 
 /* FWmfMediaPlayer callbacks
  *****************************************************************************/
 
-void FWmfMediaPlayer::HandleSessionError(HRESULT Error)
-{
-	UE_LOG(LogWmfMedia, Error, TEXT("An error occured in the media session: 0x%X"), Error);
-}
-
-
 void FWmfMediaPlayer::HandleSessionEvent(MediaEventType EventType)
 {
 	TOptional<EMediaEvent> Event;
@@ -513,19 +681,36 @@ void FWmfMediaPlayer::HandleSessionEvent(MediaEventType EventType)
 		Event = EMediaEvent::PlaybackEndReached;
 		break;
 
+	case MESessionClosed:
+	case MEError:
+		Tracks.SetPaused(true);
+		Tracks.Flush();
+		break;
+
 	case MESessionStarted:
-		Event = EMediaEvent::PlaybackResumed;
+		if (!FMath::IsNearlyZero(MediaSession->GetRate()))
+		{
+			Event = EMediaEvent::PlaybackResumed;
+			Tracks.SetPaused(false);
+		}
+		break;
+
+	case MESessionEnded:
+		Event = EMediaEvent::PlaybackSuspended;
+		Tracks.SetPaused(true);
+		Tracks.Flush();
 		break;
 
 	case MESessionStopped:
 		Event = EMediaEvent::PlaybackSuspended;
+		Tracks.SetPaused(true);
 		break;
 	}
 
 	// forward event to game thread
 	if (Event.IsSet())
 	{
-		Tasks.Enqueue([=]() {
+		GameThreadTasks.Enqueue([=]() {
 			MediaEvent.Broadcast(Event.GetValue());
 		});
 	}
diff --git a/Engine/Plugins/Media/WmfMedia/Source/WmfMedia/Private/Player/WmfMediaPlayer.h b/Engine/Plugins/Media/WmfMedia/Source/WmfMedia/Private/Player/WmfMediaPlayer.h
--- a/Engine/Plugins/Media/WmfMedia/Source/WmfMedia/Private/Player/WmfMediaPlayer.h
+++ b/Engine/Plugins/Media/WmfMedia/Source/WmfMedia/Private/Player/WmfMediaPlayer.h
@@ -2,16 +2,25 @@
 
 #pragma once
 
+#include "IMediaPlayer.h"
+#include "IMediaTracks.h"
 #include "IWmfMediaResolverCallbacks.h"
+#include "WmfMediaTracks.h"
 #include "AllowWindowsPlatformTypes.h"
 
 
+class FWmfMediaResolver;
+class FWmfMediaSession;
+class IMediaControls;
+class IMediaOutput;
+class IMediaTracks;
+
+
 /**
  * Implements a media player using the Windows Media Foundation framework.
  */
 class FWmfMediaPlayer
 	: public FTickerObjectBase
-	, public IMediaInfo
 	, public IMediaPlayer
 	, protected IWmfMediaResolverCallbacks
 {
@@ -31,36 +40,18 @@ public:
 
 public:
 
-	// IMediaInfo interface
-
-	virtual FTimespan GetDuration() const override;
-	virtual TRange<float> GetSupportedRates(EMediaPlaybackDirections Direction, bool Unthinned) const override;
-	virtual FString GetUrl() const override;
-	virtual bool SupportsRate(float Rate, bool Unthinned) const override;
-	virtual bool SupportsScrubbing() const override;
-	virtual bool SupportsSeeking() const override;
-
-public:
-
-	// IMediaPlayer interface
+	//~ IMediaPlayer interface
 
 	virtual void Close() override;
-	virtual const TArray<IMediaAudioTrackRef>& GetAudioTracks() const override;
-	virtual const TArray<IMediaCaptionTrackRef>& GetCaptionTracks() const override;
-	virtual const IMediaInfo& GetMediaInfo() const override;
-	virtual float GetRate() const override;
-	virtual FTimespan GetTime() const override;
-	virtual const TArray<IMediaVideoTrackRef>& GetVideoTracks() const override;
-	virtual bool IsLooping() const override;
-	virtual bool IsPaused() const override;
-	virtual bool IsPlaying() const override;
-	virtual bool IsReady() const override;
-	virtual bool Open(const FString& Url) override;
-	virtual bool Open(const TSharedRef<FArchive, ESPMode::ThreadSafe>& Archive, const FString& OriginalUrl) override;
-	virtual bool Seek(const FTimespan& Time) override;
-	virtual bool SetLooping(bool Looping) override;
-	virtual bool SetRate(float Rate) override;
-
+	virtual IMediaControls& GetControls() override;
+	virtual FString GetInfo() const override;
+	virtual IMediaOutput& GetOutput() override;
+	virtual FString GetStats() const override;
+	virtual IMediaTracks& GetTracks() override;
+	virtual FString GetUrl() const override;
+	virtual bool Open(const FString& Url, const IMediaOptions& Options) override;
+	virtual bool Open(const TSharedRef<FArchive, ESPMode::ThreadSafe>& Archive, const FString& OriginalUrl, const IMediaOptions& Options) override;
+	
 	DECLARE_DERIVED_EVENT(FWmfMediaPlayer, IMediaPlayer::FOnMediaEvent, FOnMediaEvent);
 	virtual FOnMediaEvent& OnMediaEvent() override
 	{
@@ -69,7 +60,7 @@ public:
 
 protected:
 
-	// IWmfMediaResolverCallbacks interface
+	//~ IWmfMediaResolverCallbacks interface
 
 	virtual void ProcessResolveComplete(TComPtr<IUnknown> SourceObject, FString ResolvedUrl) override;
 	virtual void ProcessResolveFailed(FString FailedUrl) override;
@@ -83,7 +74,6 @@ protected:
 	 * @param Topology The topology to add the stream to.
 	 * @param PresentationDescriptor The presentation descriptor object.
 	 * @param MediaSourceObject The media source object.
-	 * @return true on success, false otherwise.
 	 */
 	void AddStreamToTopology(uint32 StreamIndex, IMFTopology* Topology, IMFPresentationDescriptor* PresentationDescriptor, IMFMediaSource* MediaSourceObject);
 
@@ -98,30 +88,27 @@ protected:
 
 private:
 
-	/** Handles session errors. */
-	void HandleSessionError(HRESULT Error);
-
 	/** Handles session events. */
 	void HandleSessionEvent(MediaEventType EventType);
 
 private:
 
-	/** The available audio tracks. */
-	TArray<IMediaAudioTrackRef> AudioTracks;
-
-	/** The available caption tracks. */
-	TArray<IMediaCaptionTrackRef> CaptionTracks;
-
 	/** The duration of the currently loaded media. */
 	FTimespan Duration;
 
-	/** Holds an event delegate that is invoked when a media event occurred. */
+	/** Media information string. */
+	FString Info;
+
+	/** Asynchronous tasks for the game thread. */
+	TQueue<TFunction<void()>> GameThreadTasks;
+
+	/** Event delegate that is invoked when a media event occurred. */
 	FOnMediaEvent MediaEvent;
 
-	/** Holds the asynchronous callback object for the media stream. */
+	/** Asynchronous callback object for the media stream. */
 	TComPtr<FWmfMediaSession> MediaSession;
 
-	/** Holds a pointer to the media source object. */
+	/** Pointer to the media source object. */
 	TComPtr<IMFMediaSource> MediaSource;
 
 	/** The URL of the currently opened media. */
@@ -130,11 +117,8 @@ private:
 	/** The media source resolver. */
 	TComPtr<FWmfMediaResolver> Resolver;
 
-	/** Queued up asynchronous tasks. */
-	TQueue<TFunction<void()>> Tasks;
-
-	/** The available video tracks. */
-	TArray<IMediaVideoTrackRef> VideoTracks;
+	/** Track manager. */
+	FWmfMediaTracks Tracks;
 };
 
 
diff --git a/Engine/Plugins/Media/WmfMedia/Source/WmfMedia/Private/Player/WmfMediaTracks.cpp b/Engine/Plugins/Media/WmfMedia/Source/WmfMedia/Private/Player/WmfMediaTracks.cpp
new file mode 100644
--- /dev/null
+++ b/Engine/Plugins/Media/WmfMedia/Source/WmfMedia/Private/Player/WmfMediaTracks.cpp
@@ -0,0 +1,742 @@
+// Copyright 1998-2016 Epic Games, Inc. All Rights Reserved.
+
+#include "WmfMediaPCH.h"
+#include "WmfMediaSampler.h"
+#include "WmfMediaTracks.h"
+#include "AllowWindowsPlatformTypes.h"
+
+
+#define LOCTEXT_NAMESPACE "FWmfMediaTracks"
+
+
+/* FWmfMediaTracks structors
+ *****************************************************************************/
+
+FWmfMediaTracks::FWmfMediaTracks()
+	: SelectedAudioTrack(INDEX_NONE)
+	, SelectedCaptionTrack(INDEX_NONE)
+	, SelectedImageTrack(INDEX_NONE)
+	, SelectedVideoTrack(INDEX_NONE)
+	, AudioSink(nullptr)
+	, CaptionSink(nullptr)
+	, ImageSink(nullptr)
+	, VideoSink(nullptr)
+{ }
+
+
+FWmfMediaTracks::~FWmfMediaTracks()
+{
+	Reset();
+}
+
+
+/* FWmfMediaTracks interface
+ *****************************************************************************/
+
+void FWmfMediaTracks::AddTrack(const GUID& MajorType, IMFMediaType* MediaType, IMFMediaType* OutputType, FWmfMediaSampler* Sampler, IMFStreamDescriptor* StreamDescriptor, DWORD StreamIndex)
+{
+	check(PresentationDescriptor != nullptr);
+
+	FScopeLock Lock(&CriticalSection);
+
+	// get stream info
+	FString Language;
+	{
+		PWSTR LanguageString = NULL;
+		UINT32 OutLength;
+
+		if (SUCCEEDED(StreamDescriptor->GetAllocatedString(MF_SD_LANGUAGE, &LanguageString, &OutLength)))
+		{
+			Language = LanguageString;
+			CoTaskMemFree(LanguageString);
+		}
+	}
+
+	FString Name;
+	{
+		PWSTR NameString = NULL;
+		UINT32 OutLength;
+
+		if (SUCCEEDED(StreamDescriptor->GetAllocatedString(MF_SD_STREAM_NAME, &NameString, &OutLength)))
+		{
+			Name = NameString;
+			CoTaskMemFree(NameString);
+		}
+	}
+
+	FText DisplayName = (Name.IsEmpty())
+		? FText::Format(LOCTEXT("UnnamedStreamFormat", "Unnamed Stream {0}"), FText::AsNumber((uint32)StreamIndex))
+		: FText::FromString(Name);
+		
+	bool Protected = ::MFGetAttributeUINT32(StreamDescriptor, MF_SD_PROTECTED, FALSE) != 0;
+		
+	// add track
+	if (MajorType == MFMediaType_Audio)
+	{
+		const int32 AudioTrackIndex = AudioTracks.AddDefaulted();
+		FAudioTrack& AudioTrack = AudioTracks[AudioTrackIndex];
+		{
+			AudioTrack.DisplayName = DisplayName;
+			AudioTrack.Language = Language;
+			AudioTrack.Name = Name;
+			AudioTrack.NumChannels = ::MFGetAttributeUINT32(MediaType, MF_MT_AUDIO_NUM_CHANNELS, 0);
+			AudioTrack.Protected = Protected;
+			AudioTrack.SampleRate = ::MFGetAttributeUINT32(MediaType, MF_MT_AUDIO_SAMPLES_PER_SECOND, 0);
+			AudioTrack.Sampler = Sampler;
+			AudioTrack.StreamDescriptor = StreamDescriptor;
+			AudioTrack.StreamIndex = StreamIndex;
+		}
+
+		Sampler->OnClock().AddRaw(this, &FWmfMediaTracks::HandleMediaSamplerClock, EMediaTrackType::Audio);
+		Sampler->OnSample().AddRaw(this, &FWmfMediaTracks::HandleMediaSamplerSample, EMediaTrackType::Audio);
+	}
+	else if (MajorType == MFMediaType_Image)
+	{
+		const int32 ImageTrackIndex = ImageTracks.AddDefaulted();
+		FImageTrack& ImageTrack = ImageTracks[ImageTrackIndex];
+		{
+			ImageTrack.DisplayName = DisplayName;
+			ImageTrack.Language = Language;
+			ImageTrack.Name = Name;
+			ImageTrack.Protected = Protected;
+			ImageTrack.Sampler = Sampler;
+			ImageTrack.StreamDescriptor = StreamDescriptor;
+			ImageTrack.StreamIndex = StreamIndex;
+			
+			::MFGetAttributeSize(MediaType, MF_MT_FRAME_SIZE, (UINT32*)&ImageTrack.Dimensions.X, (UINT32*)&ImageTrack.Dimensions.Y);
+		}
+			
+		Sampler->OnSample().AddRaw(this, &FWmfMediaTracks::HandleMediaSamplerSample, EMediaTrackType::Image);
+	}
+	else if (MajorType == MFMediaType_SAMI)
+	{
+		const int32 CaptionTrackIndex = CaptionTracks.AddDefaulted();
+		FCaptionTrack& CaptionTrack = CaptionTracks[CaptionTrackIndex];
+		{
+			CaptionTrack.DisplayName = DisplayName;
+			CaptionTrack.Language = Language;
+			CaptionTrack.Name = Name;
+			CaptionTrack.Protected = Protected;
+			CaptionTrack.Sampler = Sampler;
+			CaptionTrack.StreamDescriptor = StreamDescriptor;
+			CaptionTrack.StreamIndex = StreamIndex;
+		}
+
+		Sampler->OnSample().AddRaw(this, &FWmfMediaTracks::HandleMediaSamplerSample, EMediaTrackType::Caption);
+	}
+	else if (MajorType == MFMediaType_Video)
+	{
+		const int32 VideoTrackIndex = VideoTracks.AddDefaulted();
+		FVideoTrack& VideoTrack = VideoTracks[VideoTrackIndex];
+		{
+			VideoTrack.BitRate = ::MFGetAttributeUINT32(MediaType, MF_MT_AVG_BITRATE, 0);
+			VideoTrack.DisplayName = DisplayName;
+			VideoTrack.Language = Language;
+			VideoTrack.Name = Name;
+			VideoTrack.Protected = Protected;
+			VideoTrack.Sampler = Sampler;
+			VideoTrack.StreamDescriptor = StreamDescriptor;
+			VideoTrack.StreamIndex = StreamIndex;
+
+			::MFGetAttributeSize(MediaType, MF_MT_FRAME_SIZE, (UINT32*)&VideoTrack.Dimensions.X, (UINT32*)&VideoTrack.Dimensions.Y);
+
+			// frame rate
+			UINT32 Numerator = 0;
+			UINT32 Denominator = 1;
+
+			VideoTrack.FrameRate = SUCCEEDED(::MFGetAttributeRatio(MediaType, MF_MT_FRAME_RATE, &Numerator, &Denominator))
+				? static_cast<float>(Numerator) / Denominator
+				: 0.0f;
+
+			// sink format
+			GUID SubType;
+
+			if (SUCCEEDED(OutputType->GetGUID(MF_MT_SUBTYPE, &SubType)) && ((SubType == MFVideoFormat_YUY2)))
+			{
+				VideoTrack.Pitch = VideoTrack.Dimensions.X * 2;
+				VideoTrack.SinkFormat = EMediaTextureSinkFormat::CharYUY2;
+			}
+			else
+			{
+				VideoTrack.Pitch = VideoTrack.Dimensions.X * 4;
+				VideoTrack.SinkFormat = EMediaTextureSinkFormat::CharBGRA;
+			}
+		}
+
+		Sampler->OnSample().AddRaw(this, &FWmfMediaTracks::HandleMediaSamplerSample, EMediaTrackType::Video);
+	}
+}
+
+
+void FWmfMediaTracks::AppendStats(FString &OutStats) const
+{
+	FScopeLock Lock(&CriticalSection);
+
+	// audio tracks
+	OutStats += TEXT("Audio Tracks\n");
+	
+	if (AudioTracks.Num() == 0)
+	{
+		OutStats += TEXT("    none\n");
+	}
+	else
+	{
+		for (const FAudioTrack& AudioTrack : AudioTracks)
+		{
+			OutStats += FString::Printf(TEXT("    %s\n"), *AudioTrack.DisplayName.ToString());
+			OutStats += FString::Printf(TEXT("        Channels: %i\n"), AudioTrack.NumChannels);
+			OutStats += FString::Printf(TEXT("        Sample Rate: %i\n\n"), AudioTrack.SampleRate);
+		}
+	}
+
+	// video tracks
+	OutStats += TEXT("Video Tracks\n");
+
+	if (VideoTracks.Num() == 0)
+	{
+		OutStats += TEXT("    none\n");
+	}
+	else
+	{
+		for (const FVideoTrack& VideoTrack : VideoTracks)
+		{
+			OutStats += FString::Printf(TEXT("    %s\n"), *VideoTrack.DisplayName.ToString());
+			OutStats += FString::Printf(TEXT("        BitRate: %i\n"), VideoTrack.BitRate);
+			OutStats += FString::Printf(TEXT("        Frame Rate: %i\n\n"), VideoTrack.FrameRate);
+		}
+	}
+}
+
+
+void FWmfMediaTracks::Flush()
+{
+	FScopeLock Lock(&CriticalSection);
+
+	if (AudioSink != nullptr)
+	{
+		AudioSink->FlushAudioSink();
+	}
+}
+
+
+void FWmfMediaTracks::Initialize(IMFPresentationDescriptor& InPresentationDescriptor)
+{
+	PresentationDescriptor = &InPresentationDescriptor;
+}
+
+
+void FWmfMediaTracks::Reset()
+{
+	FScopeLock Lock(&CriticalSection);
+
+	SelectedAudioTrack = INDEX_NONE;
+	SelectedCaptionTrack = INDEX_NONE;
+	SelectedImageTrack = INDEX_NONE;
+	SelectedVideoTrack = INDEX_NONE;
+
+	if (AudioSink != nullptr)
+	{
+		AudioSink->ShutdownAudioSink();
+	}
+
+	if (CaptionSink != nullptr)
+	{
+		CaptionSink->ShutdownStringSink();
+	}
+
+	if (ImageSink != nullptr)
+	{
+		ImageSink->ShutdownTextureSink();
+	}
+
+	if (VideoSink != nullptr)
+	{
+		VideoSink->ShutdownTextureSink();
+	}
+
+	AudioTracks.Empty();
+	CaptionTracks.Empty();
+	ImageTracks.Empty();
+	VideoTracks.Empty();
+
+	PresentationDescriptor = nullptr;
+}
+
+
+void FWmfMediaTracks::SetPaused(bool Paused)
+{
+	FScopeLock Lock(&CriticalSection);
+
+	if (AudioSink != nullptr)
+	{
+		if (Paused)
+		{
+			AudioSink->PauseAudioSink();
+		}
+		else
+		{
+			AudioSink->ResumeAudioSink();
+		}
+	}
+}
+
+
+/* IMediaOutput interface
+ *****************************************************************************/
+
+void FWmfMediaTracks::SetAudioSink(IMediaAudioSink* Sink)
+{
+	if (Sink != AudioSink)
+	{
+		FScopeLock Lock(&CriticalSection);
+
+		if (AudioSink != nullptr)
+		{
+			AudioSink->ShutdownAudioSink();
+		}
+
+		AudioSink = Sink;	
+		InitializeAudioSink();
+	}
+}
+
+
+void FWmfMediaTracks::SetCaptionSink(IMediaStringSink* Sink)
+{
+	if (Sink != CaptionSink)
+	{
+		FScopeLock Lock(&CriticalSection);
+
+		if (CaptionSink != nullptr)
+		{
+			CaptionSink->ShutdownStringSink();
+		}
+
+		CaptionSink = Sink;
+		InitializeCaptionSink();
+	}
+}
+
+
+void FWmfMediaTracks::SetImageSink(IMediaTextureSink* Sink)
+{
+	if (Sink != ImageSink)
+	{
+		FScopeLock Lock(&CriticalSection);
+
+		if (ImageSink != nullptr)
+		{
+			ImageSink->ShutdownTextureSink();
+		}
+
+		ImageSink = Sink;
+		InitializeImageSink();
+	}
+}
+
+
+void FWmfMediaTracks::SetVideoSink(IMediaTextureSink* Sink)
+{
+	if (Sink != VideoSink)
+	{
+		FScopeLock Lock(&CriticalSection);
+
+		if (VideoSink != nullptr)
+		{
+			VideoSink->ShutdownTextureSink();
+		}
+
+		VideoSink = Sink;
+		InitializeVideoSink();
+	}
+}
+
+
+/* IMediaTracks interface
+ *****************************************************************************/
+
+uint32 FWmfMediaTracks::GetAudioTrackChannels(int32 TrackIndex) const
+{
+	return AudioTracks.IsValidIndex(TrackIndex) ? AudioTracks[TrackIndex].NumChannels : 0;
+}
+
+
+uint32 FWmfMediaTracks::GetAudioTrackSampleRate(int32 TrackIndex) const
+{
+	return AudioTracks.IsValidIndex(TrackIndex) ? AudioTracks[TrackIndex].SampleRate : 0;
+}
+
+
+int32 FWmfMediaTracks::GetNumTracks(EMediaTrackType TrackType) const
+{
+	switch (TrackType)
+	{
+	case EMediaTrackType::Audio:
+		return AudioTracks.Num();
+	case EMediaTrackType::Caption:
+		return CaptionTracks.Num();
+	case EMediaTrackType::Video:
+		return VideoTracks.Num();
+	default:
+		return 0;
+	}
+}
+
+
+int32 FWmfMediaTracks::GetSelectedTrack(EMediaTrackType TrackType) const
+{
+	switch (TrackType)
+	{
+	case EMediaTrackType::Audio:
+		return SelectedAudioTrack;
+	case EMediaTrackType::Caption:
+		return SelectedCaptionTrack;
+	case EMediaTrackType::Video:
+		return SelectedVideoTrack;
+	default:
+		return INDEX_NONE;
+	}
+}
+
+
+FText FWmfMediaTracks::GetTrackDisplayName(EMediaTrackType TrackType, int32 TrackIndex) const
+{
+	switch (TrackType)
+	{
+	case EMediaTrackType::Audio:
+		if (AudioTracks.IsValidIndex(TrackIndex))
+		{
+			return AudioTracks[TrackIndex].DisplayName;
+		}
+		break;
+
+	case EMediaTrackType::Caption:
+		if (CaptionTracks.IsValidIndex(TrackIndex))
+		{
+			return CaptionTracks[TrackIndex].DisplayName;
+		}
+		break;
+
+	case EMediaTrackType::Video:
+		if (VideoTracks.IsValidIndex(TrackIndex))
+		{
+			return VideoTracks[TrackIndex].DisplayName;
+		}
+		break;
+
+	default:
+		break;
+	}
+
+	return FText::GetEmpty();
+}
+
+
+FString FWmfMediaTracks::GetTrackLanguage(EMediaTrackType TrackType, int32 TrackIndex) const
+{
+	switch (TrackType)
+	{
+	case EMediaTrackType::Audio:
+		if (AudioTracks.IsValidIndex(TrackIndex))
+		{
+			return AudioTracks[TrackIndex].Language;
+		}
+		break;
+
+	case EMediaTrackType::Caption:
+		if (CaptionTracks.IsValidIndex(TrackIndex))
+		{
+			return CaptionTracks[TrackIndex].Language;
+		}
+		break;
+
+	case EMediaTrackType::Video:
+		if (VideoTracks.IsValidIndex(TrackIndex))
+		{
+			return VideoTracks[TrackIndex].Language;
+		}
+		break;
+
+	default:
+		break;
+	}
+
+	return FString();
+}
+
+
+FString FWmfMediaTracks::GetTrackName(EMediaTrackType TrackType, int32 TrackIndex) const
+{
+	switch (TrackType)
+	{
+	case EMediaTrackType::Audio:
+		if (AudioTracks.IsValidIndex(TrackIndex))
+		{
+			return AudioTracks[TrackIndex].Name;
+		}
+		break;
+
+	case EMediaTrackType::Caption:
+		if (CaptionTracks.IsValidIndex(TrackIndex))
+		{
+			return CaptionTracks[TrackIndex].Name;
+		}
+		break;
+
+	case EMediaTrackType::Video:
+		if (VideoTracks.IsValidIndex(TrackIndex))
+		{
+			return VideoTracks[TrackIndex].Name;
+		}
+		break;
+
+	default:
+		break;
+	}
+
+	return FString();
+}
+
+
+uint32 FWmfMediaTracks::GetVideoTrackBitRate(int32 TrackIndex) const
+{
+	return VideoTracks.IsValidIndex(TrackIndex) ? VideoTracks[TrackIndex].BitRate : 0;
+}
+
+
+FIntPoint FWmfMediaTracks::GetVideoTrackDimensions(int32 TrackIndex) const
+{
+	return VideoTracks.IsValidIndex(TrackIndex) ? VideoTracks[TrackIndex].Dimensions : FIntPoint::ZeroValue;
+}
+
+
+float FWmfMediaTracks::GetVideoTrackFrameRate(int32 TrackIndex) const
+{
+	return VideoTracks.IsValidIndex(TrackIndex) ? VideoTracks[TrackIndex].FrameRate : 0.0f;
+}
+
+
+bool FWmfMediaTracks::SelectTrack(EMediaTrackType TrackType, int32 TrackIndex)
+{
+	FScopeLock Lock(&CriticalSection);
+
+	switch (TrackType)
+	{
+	case EMediaTrackType::Audio:
+		if ((TrackIndex == INDEX_NONE) || AudioTracks.IsValidIndex(TrackIndex))
+		{
+			if (TrackIndex != SelectedAudioTrack)
+			{
+				if (SelectedAudioTrack != INDEX_NONE)
+				{
+					PresentationDescriptor->DeselectStream(AudioTracks[SelectedAudioTrack].StreamIndex);
+					SelectedAudioTrack = INDEX_NONE;
+				}
+
+				if ((TrackIndex != INDEX_NONE) && SUCCEEDED(PresentationDescriptor->SelectStream(AudioTracks[TrackIndex].StreamIndex)))
+				{
+					SelectedAudioTrack = TrackIndex;
+				}
+
+				if (SelectedAudioTrack == TrackIndex)
+				{
+					InitializeAudioSink();
+				}
+			}
+
+			return true;
+		}
+		break;
+
+	case EMediaTrackType::Caption:
+		if ((TrackIndex == INDEX_NONE) || CaptionTracks.IsValidIndex(TrackIndex))
+		{
+			if (TrackIndex != SelectedCaptionTrack)
+			{
+				if (SelectedCaptionTrack != INDEX_NONE)
+				{
+					PresentationDescriptor->DeselectStream(CaptionTracks[SelectedCaptionTrack].StreamIndex);
+					SelectedCaptionTrack = INDEX_NONE;
+				}
+
+				if ((TrackIndex != INDEX_NONE) && SUCCEEDED(PresentationDescriptor->SelectStream(CaptionTracks[TrackIndex].StreamIndex)))
+				{
+					SelectedCaptionTrack = TrackIndex;
+				}
+
+				if (SelectedCaptionTrack == TrackIndex)
+				{
+					InitializeCaptionSink();
+				}
+			}
+
+			return true;
+		}
+		break;
+
+	case EMediaTrackType::Image:
+		if ((TrackIndex == INDEX_NONE) || ImageTracks.IsValidIndex(TrackIndex))
+		{
+			if (TrackIndex != SelectedImageTrack)
+			{
+				if (SelectedImageTrack != INDEX_NONE)
+				{
+					PresentationDescriptor->DeselectStream(ImageTracks[SelectedImageTrack].StreamIndex);
+					SelectedImageTrack = INDEX_NONE;
+				}
+
+				if ((TrackIndex != INDEX_NONE) && SUCCEEDED(PresentationDescriptor->SelectStream(ImageTracks[TrackIndex].StreamIndex)))
+				{
+					SelectedImageTrack = TrackIndex;
+				}
+
+				if (SelectedImageTrack == TrackIndex)
+				{
+					InitializeImageSink();
+				}
+			}
+
+			return true;
+		}
+		break;
+
+	case EMediaTrackType::Video:
+		if ((TrackIndex == INDEX_NONE) || VideoTracks.IsValidIndex(TrackIndex))
+		{
+			if (TrackIndex != SelectedVideoTrack)
+			{
+				if (SelectedVideoTrack != INDEX_NONE)
+				{
+					PresentationDescriptor->DeselectStream(VideoTracks[SelectedVideoTrack].StreamIndex);
+					SelectedVideoTrack = INDEX_NONE;
+				}
+
+				if ((TrackIndex != INDEX_NONE) && SUCCEEDED(PresentationDescriptor->SelectStream(VideoTracks[TrackIndex].StreamIndex)))
+				{
+					SelectedVideoTrack = TrackIndex;
+				}
+
+				if (SelectedVideoTrack == TrackIndex)
+				{
+					InitializeVideoSink();
+				}
+			}
+
+			return true;
+		}
+		break;
+
+	default:
+		break;
+	}
+
+	return false;
+}
+
+
+/* FWmfMediaTracks implementation
+ *****************************************************************************/
+
+void FWmfMediaTracks::InitializeAudioSink()
+{
+	if ((AudioSink == nullptr) || (SelectedAudioTrack == INDEX_NONE))
+	{
+		return;
+	}
+
+	const FAudioTrack& AudioTrack = AudioTracks[SelectedAudioTrack];
+	AudioSink->InitializeAudioSink(AudioTrack.NumChannels, AudioTrack.SampleRate);
+}
+
+
+void FWmfMediaTracks::InitializeCaptionSink()
+{
+	if ((CaptionSink == nullptr) || (SelectedCaptionTrack == INDEX_NONE))
+	{
+		return;
+	}
+
+	CaptionSink->InitializeStringSink();
+}
+
+
+void FWmfMediaTracks::InitializeImageSink()
+{
+	if ((ImageSink == nullptr) || (SelectedImageTrack == INDEX_NONE))
+	{
+		return;
+	}
+
+	const FImageTrack& ImageTrack = ImageTracks[SelectedImageTrack];
+	ImageSink->InitializeTextureSink(ImageTrack.Dimensions, EMediaTextureSinkFormat::CharBGRA, EMediaTextureSinkMode::Buffered);
+}
+
+
+void FWmfMediaTracks::InitializeVideoSink()
+{
+	if ((VideoSink == nullptr) || (SelectedVideoTrack == INDEX_NONE))
+	{
+		return;
+	}
+
+	const FVideoTrack& VideoTrack = VideoTracks[SelectedVideoTrack];
+	VideoSink->InitializeTextureSink(VideoTrack.Dimensions, VideoTrack.SinkFormat, EMediaTextureSinkMode::Buffered);
+}
+
+
+/* FWmfMediaTracks event handlers
+ *****************************************************************************/
+
+void FWmfMediaTracks::HandleMediaSamplerClock(EWmfMediaSamplerClockEvent Event, EMediaTrackType TrackType)
+{
+	// IMFSampleGrabberSinkCallback callbacks seem to be broken (always returns Stopped)
+	// We handle sink synchronization via SetPaused() as a workaround
+}
+
+
+void FWmfMediaTracks::HandleMediaSamplerSample(const uint8* Buffer, uint32 Size, FTimespan Duration, FTimespan Time, EMediaTrackType TrackType)
+{
+	FScopeLock Lock(&CriticalSection);
+
+	if (Buffer == nullptr)
+	{
+		return;
+	}
+
+	// forward the sample to the sink
+	switch (TrackType)
+	{
+	case EMediaTrackType::Audio:
+		if ((AudioSink != nullptr) && AudioTracks.IsValidIndex(SelectedAudioTrack))
+		{
+			AudioSink->PlayAudioSink(Buffer, Size, Time);
+		}
+		break;
+
+	case EMediaTrackType::Caption:
+		if ((CaptionSink != nullptr) && CaptionTracks.IsValidIndex(SelectedCaptionTrack))
+		{
+			CaptionSink->DisplayStringSinkString((TCHAR*)Buffer, Time);
+		}
+		break;
+
+	case EMediaTrackType::Image:
+		if ((ImageSink != nullptr) && ImageTracks.IsValidIndex(SelectedImageTrack))
+		{
+			ImageSink->UpdateTextureSinkBuffer(Buffer, ImageTracks[SelectedImageTrack].Dimensions.X * 4);
+			ImageSink->DisplayTextureSinkBuffer(Time);
+		}
+		break;
+
+	case EMediaTrackType::Video:
+		if ((VideoSink != nullptr) && VideoTracks.IsValidIndex(SelectedVideoTrack))
+		{
+			VideoSink->UpdateTextureSinkBuffer(Buffer, VideoTracks[SelectedVideoTrack].Pitch);
+			VideoSink->DisplayTextureSinkBuffer(Time);
+		}
+		break;
+	}
+}
+
+
+#include "HideWindowsPlatformTypes.h"
diff --git a/Engine/Plugins/Media/WmfMedia/Source/WmfMedia/Private/Player/WmfMediaTracks.h b/Engine/Plugins/Media/WmfMedia/Source/WmfMedia/Private/Player/WmfMediaTracks.h
new file mode 100644
--- /dev/null
+++ b/Engine/Plugins/Media/WmfMedia/Source/WmfMedia/Private/Player/WmfMediaTracks.h
@@ -0,0 +1,210 @@
+// Copyright 1998-2016 Epic Games, Inc. All Rights Reserved.
+
+#pragma once
+
+#include "IMediaOutput.h"
+#include "IMediaTracks.h"
+#include "AllowWindowsPlatformTypes.h"
+
+
+class FWmfMediaAudioStream;
+class FWmfMediaCaptionStream;
+class FWmfMediaImageStream;
+class FWmfMediaSampler;
+class FWmfMediaVideoStream;
+class IMediaAudioSink;
+class IMediaStringSink;
+class IMediaTextureSink;
+enum class EWmfMediaSamplerClockEvent;
+enum class EMediaTrackType;
+
+
+class FWmfMediaTracks
+	: public IMediaOutput
+	, public IMediaTracks
+{
+	struct FTrack
+	{
+		FText DisplayName;
+		FString Language;
+		FString Name;
+		bool Protected;
+		TComPtr<FWmfMediaSampler> Sampler;
+		TComPtr<IMFStreamDescriptor> StreamDescriptor;
+		DWORD StreamIndex;	
+	};
+	
+	struct FAudioTrack : public FTrack
+	{
+		uint32 NumChannels;
+		uint32 SampleRate;
+	};
+	
+	struct FCaptionTrack : public FTrack
+	{
+	};
+	
+	struct FImageTrack : public FTrack
+	{
+		FIntPoint Dimensions;
+	};
+	
+	struct FVideoTrack : public FTrack
+	{
+		uint32 BitRate;
+		FIntPoint Dimensions;
+		float FrameRate;
+		uint32 Pitch;
+		EMediaTextureSinkFormat SinkFormat;
+	};
+	
+public:
+
+	/** Default constructor. */
+	FWmfMediaTracks();
+
+	/** Virtual destructor. */
+	~FWmfMediaTracks();
+
+public:
+
+	/**
+	 * Add a new track for the specified stream.
+	 *
+	 * @param MajorType The type of stream to add.
+	 * @param MediaType The media type information of the input stream.
+	 * @param OutputType The media type information of the output.
+	 * @param Sampler The sample grabber callback object to use.
+	 * @param StreamIndex The stream's index number in the presentation.
+	 * @param StreamDescriptor The stream's descriptor object.
+	 * @see HasAnyStreams
+	 */
+	void AddTrack(const GUID& MajorType, IMFMediaType* MediaType, IMFMediaType* OutputType, FWmfMediaSampler* Sampler, IMFStreamDescriptor* StreamDescriptor, DWORD StreamIndex);
+
+	/**
+	 * Append track statistics information to the given string.
+	 *
+	 * @param OutStats The string to append the statistics to.
+	 */
+	void AppendStats(FString &OutStats) const;
+
+	/** Discard any pending data on the media sinks. */
+	void Flush();
+
+	/**
+	 * Initialize this track collection.
+	 *
+	 * @param InPresentationDescriptor The presentation descriptor object.
+	 */
+	void Initialize(IMFPresentationDescriptor& InPresentationDescriptor);
+
+	/** Remove all streams from the collection. */
+	void Reset();
+
+	/**
+	 * Notify tracks that playback was paused or resumed.
+	 *
+	 * @param Paused true if playback was paused, false if resumed.
+	 */
+	void SetPaused(bool Paused);
+
+public:
+
+	//~ IMediaOutput interface
+
+	virtual void SetAudioSink(IMediaAudioSink* Sink) override;
+	virtual void SetCaptionSink(IMediaStringSink* Sink) override;
+	virtual void SetImageSink(IMediaTextureSink* Sink) override;
+	virtual void SetVideoSink(IMediaTextureSink* Sink) override;
+
+public:
+
+	//~ IMediaTracks interface
+
+	virtual uint32 GetAudioTrackChannels(int32 TrackIndex) const override;
+	virtual uint32 GetAudioTrackSampleRate(int32 TrackIndex) const override;
+	virtual int32 GetNumTracks(EMediaTrackType TrackType) const override;
+	virtual int32 GetSelectedTrack(EMediaTrackType TrackType) const override;
+	virtual FText GetTrackDisplayName(EMediaTrackType TrackType, int32 TrackIndex) const override;
+	virtual FString GetTrackLanguage(EMediaTrackType TrackType, int32 TrackIndex) const override;
+	virtual FString GetTrackName(EMediaTrackType TrackType, int32 TrackIndex) const override;
+	virtual uint32 GetVideoTrackBitRate(int32 TrackIndex) const override;
+	virtual FIntPoint GetVideoTrackDimensions(int32 TrackIndex) const override;
+	virtual float GetVideoTrackFrameRate(int32 TrackIndex) const override;
+	virtual bool SelectTrack(EMediaTrackType TrackType, int32 TrackIndex) override;
+
+protected:
+
+	/** Initialize the current audio sink. */
+	void InitializeAudioSink();
+
+	/** Initialize the current caption sink. */
+	void InitializeCaptionSink();
+
+	/** Initialize the current image sink. */
+	void InitializeImageSink();
+
+	/** Initialize the current video sink. */
+	void InitializeVideoSink();
+
+private:
+
+	/** Callback for handling media sampler pauses. */
+	void HandleMediaSamplerClock(EWmfMediaSamplerClockEvent Event, EMediaTrackType TrackType);
+
+	/** Callback for handling new samples from the streams' media sample buffers. */
+	void HandleMediaSamplerSample(const uint8* Buffer, uint32 Size, FTimespan Duration, FTimespan Time, EMediaTrackType TrackType);
+
+private:
+
+	/** The currently used audio sink. */
+	IMediaAudioSink* AudioSink;
+
+	/** The currently used caption sink. */
+	IMediaStringSink* CaptionSink;
+
+	/** The currently used image sink. */
+	IMediaTextureSink* ImageSink;
+
+	/** The currently used video sink. */
+	IMediaTextureSink* VideoSink;
+
+private:
+
+	/** The available audio tracks. */
+	TArray<FAudioTrack> AudioTracks;
+
+	/** The available caption tracks. */
+	TArray<FCaptionTrack> CaptionTracks;
+
+	/** The available image tracks. */
+	TArray<FImageTrack> ImageTracks;
+
+	/** The available video tracks. */
+	TArray<FVideoTrack> VideoTracks;
+
+private:
+
+	/** Index of the selected audio track. */
+	int32 SelectedAudioTrack;
+
+	/** Index of the selected caption track. */
+	int32 SelectedCaptionTrack;
+
+	/** Index of the selected image track. */
+	int32 SelectedImageTrack;
+
+	/** Index of the selected video track. */
+	int32 SelectedVideoTrack;
+
+private:
+
+	/** Synchronizes write access to track arrays, selections & sinks. */
+	mutable FCriticalSection CriticalSection;
+
+	/** The presentation descriptor of the currently opened media. */
+	TComPtr<IMFPresentationDescriptor> PresentationDescriptor;
+};
+
+
+#include "HideWindowsPlatformTypes.h"
diff --git a/Engine/Plugins/Media/WmfMedia/Source/WmfMedia/Private/Wmf/WmfMediaByteStream.cpp b/Engine/Plugins/Media/WmfMedia/Source/WmfMedia/Private/Wmf/WmfMediaByteStream.cpp
--- a/Engine/Plugins/Media/WmfMedia/Source/WmfMedia/Private/Wmf/WmfMediaByteStream.cpp
+++ b/Engine/Plugins/Media/WmfMedia/Source/WmfMedia/Private/Wmf/WmfMediaByteStream.cpp
@@ -1,6 +1,8 @@
 // Copyright 1998-2016 Epic Games, Inc. All Rights Reserved.
 
-#include "WmfMediaPrivatePCH.h"
+#include "WmfMediaPCH.h"
+#include "WmfMediaByteStream.h"
+#include "WmfMediaReadState.h"
 #include "AllowWindowsPlatformTypes.h"
 
 
@@ -33,19 +35,24 @@ STDMETHODIMP FWmfMediaByteStream::Invoke(IMFAsyncResult* AsyncResult)
 {
 	// recover read state
 	IUnknown* State = NULL;
-	HRESULT Result = AsyncResult->GetState(&State);
-	check(SUCCEEDED(Result));
+	
+	if (FAILED(AsyncResult->GetState(&State)))
+	{
+		return S_OK;
+	}
 
 	IMFAsyncResult* CallerResult = NULL;
-	Result = State->QueryInterface(IID_PPV_ARGS(&CallerResult));
-	check(SUCCEEDED(Result));
+
+	if (FAILED(State->QueryInterface(IID_PPV_ARGS(&CallerResult))))
+	{
+		return S_OK;
+	}
 
 	IUnknown* Unknown = NULL;
 	
-	if (CallerResult != NULL)
+	if ((CallerResult != NULL) && FAILED(CallerResult->GetObject(&Unknown)))
 	{
-		Result = CallerResult->GetObject(&Unknown);
-		check(SUCCEEDED(Result));
+		return S_OK;
 	}
 
 	FWmfMediaReadState* ReadState = static_cast<FWmfMediaReadState*>(Unknown);
@@ -58,7 +65,7 @@ STDMETHODIMP FWmfMediaByteStream::Invoke(IMFAsyncResult* AsyncResult)
 	// notify caller
 	if (CallerResult != NULL)
 	{
-		CallerResult->SetStatus(Result);
+		CallerResult->SetStatus(S_OK);
 		::MFInvokeCallback(CallerResult);
 	}
 
@@ -81,10 +88,11 @@ STDMETHODIMP FWmfMediaByteStream::Invoke(IMFAsyncResult* AsyncResult)
 	return S_OK;
 }
 
+
 #if _MSC_VER == 1900
-#pragma warning(push)
-#pragma warning(disable:4838)
-#endif // _MSC_VER == 1900
+	#pragma warning(push)
+	#pragma warning(disable:4838)
+#endif
 
 STDMETHODIMP FWmfMediaByteStream::QueryInterface(REFIID RefID, void** Object)
 {
@@ -97,9 +105,10 @@ STDMETHODIMP FWmfMediaByteStream::QueryInterface(REFIID RefID, void** Object)
 
 	return QISearch( this, QITab, RefID, Object );
 }
+
 #if _MSC_VER == 1900
-#pragma warning(pop)
-#endif // _MSC_VER == 1900
+	#pragma warning(pop)
+#endif
 
 
 STDMETHODIMP_(ULONG) FWmfMediaByteStream::Release()
@@ -309,6 +318,12 @@ STDMETHODIMP FWmfMediaByteStream::SetCurrentPosition(QWORD qwPosition)
 		return S_FALSE;
 	}
 
+	if (qwPosition > (QWORD)Archive->TotalSize())
+	{
+		// MSDN is wrong: clamp instead of returning E_INVALIDARG
+		qwPosition = (QWORD)Archive->TotalSize();
+	}
+
 	Archive->Seek((int64)qwPosition);
 
 	return S_OK;
diff --git a/Engine/Plugins/Media/WmfMedia/Source/WmfMedia/Private/Wmf/WmfMediaByteStream.h b/Engine/Plugins/Media/WmfMedia/Source/WmfMedia/Private/Wmf/WmfMediaByteStream.h
--- a/Engine/Plugins/Media/WmfMedia/Source/WmfMedia/Private/Wmf/WmfMediaByteStream.h
+++ b/Engine/Plugins/Media/WmfMedia/Source/WmfMedia/Private/Wmf/WmfMediaByteStream.h
@@ -26,7 +26,7 @@ public:
 
 public:
 
-	// IMFAsyncCallback interface
+	//~ IMFAsyncCallback interface
 
 	STDMETHODIMP_(ULONG) AddRef();
 	STDMETHODIMP GetParameters(unsigned long*, unsigned long*);
@@ -36,7 +36,7 @@ public:
 
 public:
 
-	// IMFByteStream interface
+	//~ IMFByteStream interface
 
 	STDMETHODIMP BeginRead(BYTE* pb, ULONG cb, IMFAsyncCallback* pCallback, IUnknown* punkState);
 	STDMETHODIMP BeginWrite(const BYTE* pb, ULONG cb, IMFAsyncCallback* pCallback, IUnknown* punkState);
diff --git a/Engine/Plugins/Media/WmfMedia/Source/WmfMedia/Private/Wmf/WmfMediaReadState.h b/Engine/Plugins/Media/WmfMedia/Source/WmfMedia/Private/Wmf/WmfMediaReadState.h
--- a/Engine/Plugins/Media/WmfMedia/Source/WmfMedia/Private/Wmf/WmfMediaReadState.h
+++ b/Engine/Plugins/Media/WmfMedia/Source/WmfMedia/Private/Wmf/WmfMediaReadState.h
@@ -77,12 +77,13 @@ public:
 
 public:
 
-	// IUnknown interface
+	//~ IUnknown interface
 
 #if _MSC_VER == 1900
-#pragma warning(push)
-#pragma warning(disable:4838)
-#endif // _MSC_VER == 1900
+	#pragma warning(push)
+	#pragma warning(disable:4838)
+#endif
+
 	STDMETHODIMP QueryInterface(REFIID RefID, void** Object)
 	{
 		static const QITAB QITab[] =
@@ -93,9 +94,10 @@ public:
 
 		return QISearch( this, QITab, RefID, Object );
 	}
+
 #if _MSC_VER == 1900
-#pragma warning(pop)
-#endif // _MSC_VER == 1900
+	#pragma warning(pop)
+#endif
 
 	STDMETHODIMP_(ULONG) AddRef()
 	{
diff --git a/Engine/Plugins/Media/WmfMedia/Source/WmfMedia/Private/Wmf/WmfMediaResolveState.h b/Engine/Plugins/Media/WmfMedia/Source/WmfMedia/Private/Wmf/WmfMediaResolveState.h
--- a/Engine/Plugins/Media/WmfMedia/Source/WmfMedia/Private/Wmf/WmfMediaResolveState.h
+++ b/Engine/Plugins/Media/WmfMedia/Source/WmfMedia/Private/Wmf/WmfMediaResolveState.h
@@ -84,7 +84,7 @@ public:
 
 public:
 
-	// IUnknown interface
+	//~ IUnknown interface
 
 	STDMETHODIMP_(ULONG) AddRef()
 	{
@@ -92,9 +92,10 @@ public:
 	}
 
 #if _MSC_VER == 1900
-#pragma warning(push)
-#pragma warning(disable:4838)
-#endif // _MSC_VER == 1900
+	#pragma warning(push)
+	#pragma warning(disable:4838)
+#endif
+
 	STDMETHODIMP QueryInterface(REFIID RefID, void** Object)
 	{
 		static const QITAB QITab[] =
@@ -105,9 +106,10 @@ public:
 
 		return QISearch(this, QITab, RefID, Object);
 	}
+
 #if _MSC_VER == 1900
-#pragma warning(pop)
-#endif // _MSC_VER == 1900
+	#pragma warning(pop)
+#endif
 
 	STDMETHODIMP_(ULONG) Release()
 	{
diff --git a/Engine/Plugins/Media/WmfMedia/Source/WmfMedia/Private/Wmf/WmfMediaResolver.cpp b/Engine/Plugins/Media/WmfMedia/Source/WmfMedia/Private/Wmf/WmfMediaResolver.cpp
--- a/Engine/Plugins/Media/WmfMedia/Source/WmfMedia/Private/Wmf/WmfMediaResolver.cpp
+++ b/Engine/Plugins/Media/WmfMedia/Source/WmfMedia/Private/Wmf/WmfMediaResolver.cpp
@@ -1,6 +1,9 @@
 // Copyright 1998-2016 Epic Games, Inc. All Rights Reserved.
 
-#include "WmfMediaPrivatePCH.h"
+#include "WmfMediaPCH.h"
+#include "WmfMediaResolver.h"
+#include "WmfMediaResolveState.h"
+#include "WmfMediaUtils.h"
 #include "AllowWindowsPlatformTypes.h"
 
 
@@ -10,9 +13,11 @@
 FWmfMediaResolver::FWmfMediaResolver()
 	: RefCount(0)
 {
-	if (FAILED(::MFCreateSourceResolver(&SourceResolver)))
+	HRESULT Result = ::MFCreateSourceResolver(&SourceResolver);
+
+	if (FAILED(Result))
 	{
-		UE_LOG(LogWmfMedia, Error, TEXT("Failed to create media source resolver."));
+		UE_LOG(LogWmfMedia, Error, TEXT("Failed to create media source resolver (%s)"), *WmfMedia::ResultToString(Result));
 	}
 }
 
@@ -54,9 +59,11 @@ bool FWmfMediaResolver::ResolveByteStream(const TSharedRef<FArchive, ESPMode::Th
 
 	TComPtr<FWmfMediaByteStream> ByteStream = new FWmfMediaByteStream(Archive);
 	TComPtr<FWmfMediaResolveState> NewResolveState = new(std::nothrow) FWmfMediaResolveState(EWmfMediaResolveType::ByteStream, OriginalUrl, Callbacks);
+	HRESULT Result = SourceResolver->BeginCreateObjectFromByteStream(ByteStream, *OriginalUrl, MF_RESOLUTION_MEDIASOURCE, NULL, &CancelCookie, this, NewResolveState);
 
-	if (FAILED(SourceResolver->BeginCreateObjectFromByteStream(ByteStream, *OriginalUrl, MF_RESOLUTION_MEDIASOURCE, NULL, &CancelCookie, this, NewResolveState)))
+	if (FAILED(Result))
 	{
+		UE_LOG(LogWmfMedia, Warning, TEXT("Failed to resolve byte stream %s (%s)"), *OriginalUrl, *WmfMedia::ResultToString(Result));
 		return false;
 	}
 
@@ -76,9 +83,11 @@ bool FWmfMediaResolver::ResolveUrl(const FString& Url, IWmfMediaResolverCallback
 	Cancel();
 
 	TComPtr<FWmfMediaResolveState> NewResolveState = new(std::nothrow) FWmfMediaResolveState(EWmfMediaResolveType::Url, Url, Callbacks);
+	HRESULT Result = SourceResolver->BeginCreateObjectFromURL(*Url, MF_RESOLUTION_MEDIASOURCE, NULL, &CancelCookie, this, NewResolveState);
 
-	if (FAILED(SourceResolver->BeginCreateObjectFromURL(*Url, MF_RESOLUTION_MEDIASOURCE, NULL, &CancelCookie, this, NewResolveState)))
+	if (FAILED(Result))
 	{
+		UE_LOG(LogWmfMedia, Warning, TEXT("Failed to resolve URL %s (%s)"), *Url, *WmfMedia::ResultToString(Result));
 		return false;
 	}
 
@@ -126,23 +135,25 @@ STDMETHODIMP FWmfMediaResolver::Invoke(IMFAsyncResult* AsyncResult)
 	// get the asynchronous result
 	MF_OBJECT_TYPE ObjectType;
 	TComPtr<IUnknown> SourceObject;
-	bool Succeeded = false;
+	HRESULT Result = S_FALSE;
 
 	if (State->Type == EWmfMediaResolveType::ByteStream)
 	{
-		Succeeded = SUCCEEDED(SourceResolver->EndCreateObjectFromByteStream(AsyncResult, &ObjectType, &SourceObject));
+		Result = SourceResolver->EndCreateObjectFromByteStream(AsyncResult, &ObjectType, &SourceObject);
 	}
 	else if (State->Type == EWmfMediaResolveType::Url)
 	{
-		Succeeded = SUCCEEDED(SourceResolver->EndCreateObjectFromURL(AsyncResult, &ObjectType, &SourceObject));
+		Result = SourceResolver->EndCreateObjectFromURL(AsyncResult, &ObjectType, &SourceObject);
 	}
 
-	if (Succeeded)
+	if (SUCCEEDED(Result))
 	{
 		State->ResolveComplete(SourceObject, State->Url);
 	}
 	else
 	{
+		UE_LOG(LogWmfMedia, Warning, TEXT("Failed to finish resolve %s: %s"), *State->Url, *WmfMedia::ResultToString(Result));
+
 		State->ResolveFailed(State->Url);
 	}
 
diff --git a/Engine/Plugins/Media/WmfMedia/Source/WmfMedia/Private/Wmf/WmfMediaResolver.h b/Engine/Plugins/Media/WmfMedia/Source/WmfMedia/Private/Wmf/WmfMediaResolver.h
--- a/Engine/Plugins/Media/WmfMedia/Source/WmfMedia/Private/Wmf/WmfMediaResolver.h
+++ b/Engine/Plugins/Media/WmfMedia/Source/WmfMedia/Private/Wmf/WmfMediaResolver.h
@@ -55,7 +55,7 @@ public:
 
 public:
 
-	// IMFAsyncCallback interface
+	//~ IMFAsyncCallback interface
 
 	STDMETHODIMP_(ULONG) AddRef();
 	STDMETHODIMP GetParameters(unsigned long*, unsigned long*);
diff --git a/Engine/Plugins/Media/WmfMedia/Source/WmfMedia/Private/Wmf/WmfMediaSampler.h b/Engine/Plugins/Media/WmfMedia/Source/WmfMedia/Private/Wmf/WmfMediaSampler.h
--- a/Engine/Plugins/Media/WmfMedia/Source/WmfMedia/Private/Wmf/WmfMediaSampler.h
+++ b/Engine/Plugins/Media/WmfMedia/Source/WmfMedia/Private/Wmf/WmfMediaSampler.h
@@ -6,15 +6,24 @@
 
 
 /**
- * Implements a callback object for the sample grabber sink.
- *
- * 
+ * Clock sink events.
+ */
+enum class EWmfMediaSamplerClockEvent
+{
+	Paused,
+	Restarted,
+	Started,
+	Stopped
+};
+
+
+/**
+ * Implements a callback object for the WMF sample grabber sink.
  */
 class FWmfMediaSampler
 	: public IMFSampleGrabberSinkCallback
 {
 public:
-	DECLARE_DELEGATE(CommandDelegate)
 
 	/** Default constructor. */
 	FWmfMediaSampler()
@@ -26,33 +35,23 @@ public:
 
 public:
 
-	/**
-	 * Registers a media sink with this sampler.
-	 *
-	 * @param Sink The media sink to register.
-	 * @return true on success, false otherwise.
-	 * @see UnregisterSink
-	 */
-	bool RegisterSink(const IMediaSinkRef& Sink)
+	/** Get an event that gets fired when the sampler's presentation clock changed its state. */
+	DECLARE_EVENT_OneParam(FWmfMediaSampler, FOnClock, EWmfMediaSamplerClockEvent /*Event*/)
+	FOnClock& OnClock()
 	{
-		return Commands.Enqueue(FSimpleDelegate::CreateRaw(this, &FWmfMediaSampler::HandleRegisterSink, IMediaSinkWeakPtr(Sink)));
+		return ClockEvent;
 	}
 
-	/**
-	 * Unregisters a media sink with this sampler.
-	 *
-	 * @param Sink The media sink to unregister.
-	 * @return true on success, false otherwise.
-	 * @see RegisterSink
-	 */
-	bool UnregisterSink(const IMediaSinkRef& Sink)
+	/** Get an event that gets fired when a new sample is ready (handler must be thread-safe). */
+	DECLARE_EVENT_FourParams(FWmfMediaSampler, FOnSample, const uint8* /*Buffer*/, uint32 /*Size*/, FTimespan /*Duration*/, FTimespan /*Time*/)
+	FOnSample& OnSample()
 	{
-		return Commands.Enqueue(FSimpleDelegate::CreateRaw(this, &FWmfMediaSampler::HandleUnregisterSink, IMediaSinkWeakPtr(Sink)));
+		return SampleEvent;
 	}
 
 public:
 
-	// IMFSampleGrabberSinkCallback interface
+	//~ IMFSampleGrabberSinkCallback interface
 
 	STDMETHODIMP_(ULONG) AddRef()
 	{
@@ -61,11 +60,15 @@ public:
 
 	STDMETHODIMP OnClockPause(MFTIME SystemTime)
 	{
+		ClockEvent.Broadcast(EWmfMediaSamplerClockEvent::Paused);
+
 		return S_OK;
 	}
 
 	STDMETHODIMP OnClockRestart(MFTIME SystemTime)
 	{
+		ClockEvent.Broadcast(EWmfMediaSamplerClockEvent::Restarted);
+
 		return S_OK;
 	}
 
@@ -76,34 +79,21 @@ public:
 
 	STDMETHODIMP OnClockStart(MFTIME SystemTime, LONGLONG llClockStartOffset)
 	{
+		ClockEvent.Broadcast(EWmfMediaSamplerClockEvent::Started);
+
 		return S_OK;
 	}
 
 	STDMETHODIMP OnClockStop(MFTIME SystemTime)
 	{
+		ClockEvent.Broadcast(EWmfMediaSamplerClockEvent::Stopped);
+
 		return S_OK;
 	}
 
 	STDMETHODIMP OnProcessSample(REFGUID MajorMediaType, DWORD SampleFlags, LONGLONG SampleTime, LONGLONG SampleDuration, const BYTE* SampleBuffer, DWORD SampleSize)
 	{
-		// process pending commands
-		CommandDelegate Command;
-
-		while (Commands.Dequeue(Command))
-		{
-			Command.Execute();
-		}
-
-		// notify all registered media sinks
-		for (IMediaSinkWeakPtr& SinkPtr : Sinks)
-		{
-			IMediaSinkPtr Sink = SinkPtr.Pin();
-
-			if (Sink.IsValid())
-			{
-				Sink->ProcessMediaSample(SampleBuffer, SampleSize, FTimespan(SampleDuration), FTimespan(SampleTime));
-			}
-		}
+		SampleEvent.Broadcast((uint8*)SampleBuffer, (uint32)SampleSize, FTimespan(SampleDuration), FTimespan(SampleTime));
 
 		return S_OK;
 	}
@@ -118,11 +108,11 @@ public:
 		return S_OK;
 	}
 
-
 #if _MSC_VER == 1900
-#pragma warning(push)
-#pragma warning(disable:4838)
-#endif // _MSC_VER == 1900
+	#pragma warning(push)
+	#pragma warning(disable:4838)
+#endif
+
 	STDMETHODIMP QueryInterface(REFIID RefID, void** Object)
 	{
 		static const QITAB QITab[] =
@@ -133,9 +123,10 @@ public:
 
 		return QISearch(this, QITab, RefID, Object);
 	}
+
 #if _MSC_VER == 1900
-#pragma warning(pop)
-#endif // _MSC_VER == 1900
+	#pragma warning(pop)
+#endif
 
 	STDMETHODIMP_(ULONG) Release()
 	{
@@ -151,26 +142,14 @@ public:
 
 private:
 
-	void HandleRegisterSink(IMediaSinkWeakPtr Sink)
-	{
-		Sinks.AddUnique(Sink);
-	}
-
-	void HandleUnregisterSink(IMediaSinkWeakPtr Sink)
-	{
-		Sinks.RemoveSingle(Sink);
-	}
-
-private:
-
-	/** Holds the router command queue. */
-	TQueue<CommandDelegate, EQueueMode::Mpsc> Commands;
-
-	/** The collection of registered media sinks. */
-	TArray<IMediaSinkWeakPtr> Sinks;
+	/** Event that gets fired when the sampler's presentation clock changed state. */
+	FOnClock ClockEvent;
 
 	/** Holds a reference counter for this instance. */
 	int32 RefCount;
+
+	/** Event that gets fired when a new sample is ready. */
+	FOnSample SampleEvent;
 };
 
 
diff --git a/Engine/Plugins/Media/WmfMedia/Source/WmfMedia/Private/Wmf/WmfMediaSession.cpp b/Engine/Plugins/Media/WmfMedia/Source/WmfMedia/Private/Wmf/WmfMediaSession.cpp
--- a/Engine/Plugins/Media/WmfMedia/Source/WmfMedia/Private/Wmf/WmfMediaSession.cpp
+++ b/Engine/Plugins/Media/WmfMedia/Source/WmfMedia/Private/Wmf/WmfMediaSession.cpp
@@ -1,77 +1,120 @@
 // Copyright 1998-2016 Epic Games, Inc. All Rights Reserved.
 
-#include "WmfMediaPrivatePCH.h"
+#include "WmfMediaPCH.h"
+#include "WmfMediaSession.h"
+#include "WmfMediaUtils.h"
 #include "AllowWindowsPlatformTypes.h"
 
 
 /* FWmfMediaSession structors
  *****************************************************************************/
 
+FWmfMediaSession::FWmfMediaSession()
+	: Buffering(false)
+	, CanScrub(false)
+	, Capabilities(0)
+	, ChangeRequested(false)
+	, CurrentRate(0.0f)
+	, CurrentState(EMediaState::Closed)
+	, Duration(FTimespan::Zero())
+	, Looping(false)
+	, RefCount(0)
+	, RequestedRate(0.0f)
+	, RequestedTime(FTimespan::MinValue())
+	, StateChangePending(false)
+{ }
+
+
 FWmfMediaSession::FWmfMediaSession(const FTimespan& InDuration, const TComPtr<IMFTopology>& InTopology)
-	: Capabilities(0)
+	: Buffering(false)
+	, Capabilities(0)
 	, ChangeRequested(false)
 	, CurrentRate(0.0f)
 	, Duration(InDuration)
 	, Looping(false)
 	, RefCount(0)
-	, RequestedPosition(FTimespan::MinValue())
+	, RequestedTime(FTimespan::MinValue())
 	, RequestedRate(0.0f)
 	, StateChangePending(false)
 {
-	CurrentState = EMediaStates::Error;
+	CurrentState = EMediaState::Error;
+	
+	// create media session
+	HRESULT Result = ::MFCreateMediaSession(NULL, &MediaSession);
 
-	if (FAILED(::MFCreateMediaSession(NULL, &MediaSession)))
+	if (FAILED(Result))
 	{
-		UE_LOG(LogWmfMedia, Error, TEXT("Failed to create media session"));
+		UE_LOG(LogWmfMedia, Error, TEXT("Failed to create media session (%s)"), *WmfMedia::ResultToString(Result));
+		return;
 	}
-	else if (FAILED(MediaSession->SetTopology(0, InTopology)))
+
+	// assign the media topology
+	Result = MediaSession->SetTopology(0, InTopology);
+
+	if (FAILED(Result))
 	{
-		UE_LOG(LogWmfMedia, Error, TEXT("Failed to set topology in media session"));
+		UE_LOG(LogWmfMedia, Error, TEXT("Failed to set topology in media session (%s)"), *WmfMedia::ResultToString(Result));
+		return;
 	}
-	else if (FAILED(MediaSession->BeginGetEvent(this, NULL)))
+
+	// start media event processing
+	Result = MediaSession->BeginGetEvent(this, NULL);
+
+	if (FAILED(Result))
 	{
-		UE_LOG(LogWmfMedia, Error, TEXT("Failed to start media session event processing"));
+		UE_LOG(LogWmfMedia, Error, TEXT("Failed to start media session event processing (%s)"), *WmfMedia::ResultToString(Result));
+		return;
 	}
-	else
-	{
-		// get presentation clock, if available
-		TComPtr<IMFClock> Clock;
 
-		if (SUCCEEDED(MediaSession->GetClock(&Clock)))
-		{
-			Clock->QueryInterface(IID_PPV_ARGS(&PresentationClock));
-		}
+	// get presentation clock, if available
+	TComPtr<IMFClock> Clock;
 
-		// get rate control and support, if available
-		if (SUCCEEDED(::MFGetService(MediaSession, MF_RATE_CONTROL_SERVICE, IID_PPV_ARGS(&RateControl))) &&
-			SUCCEEDED(::MFGetService(MediaSession, MF_RATE_CONTROL_SERVICE, IID_PPV_ARGS(&RateSupport))))
-		{
-			CanScrub = SUCCEEDED(RateSupport->IsRateSupported(TRUE, 0, NULL));
-		}
+	if (SUCCEEDED(MediaSession->GetClock(&Clock)))
+	{
+		Clock->QueryInterface(IID_PPV_ARGS(&PresentationClock));
+	}
 
-		CurrentState = EMediaStates::Stopped;
+	// get rate control and support, if available
+	if (SUCCEEDED(::MFGetService(MediaSession, MF_RATE_CONTROL_SERVICE, IID_PPV_ARGS(&RateControl))) &&
+		SUCCEEDED(::MFGetService(MediaSession, MF_RATE_CONTROL_SERVICE, IID_PPV_ARGS(&RateSupport))))
+	{
+		CanScrub = SUCCEEDED(RateSupport->IsRateSupported(TRUE, 0, NULL));
 	}
 
+	CurrentState = EMediaState::Stopped;
 	RequestedState = CurrentState;
 }
 
 
-/* FWmfMediaSession interface
+/* IMediaControls interface
  *****************************************************************************/
 
-FTimespan FWmfMediaSession::GetPosition() const
+FTimespan FWmfMediaSession::GetDuration() const
 {
-	if (PresentationClock == NULL)
+	return Duration;
+}
+
+
+float FWmfMediaSession::GetRate() const
+{
+	return CurrentRate;
+}
+
+
+EMediaState FWmfMediaSession::GetState() const
+{
+	if (Buffering)
 	{
-		return FTimespan::MinValue();
+		return EMediaState::Preparing;
 	}
 
-	if (RequestedPosition >= FTimespan::Zero())
+	if ((CurrentState == EMediaState::Playing) && FMath::IsNearlyZero(CurrentRate))
 	{
-		return RequestedPosition;
+		// scrubbing counts as paused
+		return EMediaState::Paused;
 	}
-	
-	return GetInternalPosition();
+
+	return CurrentState;
 }
 
 
@@ -102,38 +145,47 @@ TRange<float> FWmfMediaSession::GetSupportedRates(EMediaPlaybackDirections Direc
 }
 
 
-bool FWmfMediaSession::IsRateSupported(float Rate, bool Unthinned) const
+FTimespan FWmfMediaSession::GetTime() const
 {
-	if (Rate == 1.0f)
+	if (PresentationClock == NULL)
 	{
-		return true;
+		return FTimespan::Zero();
 	}
 
-	if ((RateControl == NULL) || (RateSupport == NULL))
+	if (RequestedTime >= FTimespan::Zero())
 	{
-		return false;
+		return RequestedTime;
 	}
 
-	if (Unthinned)
-	{
-		return SUCCEEDED(RateSupport->IsRateSupported(FALSE, Rate, NULL));
-	}
+	return GetInternalPosition();
+}
 
-	return SUCCEEDED(RateSupport->IsRateSupported(TRUE, Rate, NULL));
+
+bool FWmfMediaSession::IsLooping() const
+{
+	return Looping;
 }
 
 
-bool FWmfMediaSession::SetPosition(const FTimespan& Position)
+bool FWmfMediaSession::Seek(const FTimespan& Time)
 {
-	RequestedPosition = Position;
+	RequestedTime = Time;
 
 	return ChangeState();
 }
 
 
+bool FWmfMediaSession::SetLooping(bool InLooping)
+{
+	Looping = InLooping;
+
+	return true;
+}
+
+
 bool FWmfMediaSession::SetRate(float Rate)
 {
-	if (!IsRateSupported(Rate, false))
+	if (!SupportsRate(Rate, false))
 	{
 		return false;
 	}
@@ -144,25 +196,61 @@ bool FWmfMediaSession::SetRate(float Rate)
 	{
 		if (!SupportsScrubbing())
 		{
-			RequestedState = EMediaStates::Paused;
+			RequestedState = EMediaState::Paused;
 		}
 	}
 	else
 	{
-		RequestedState = EMediaStates::Playing;
+		RequestedState = EMediaState::Playing;
 	}
 
 	return ChangeState();
 }
 
 
-bool FWmfMediaSession::SetState(EMediaStates NewState)
+bool FWmfMediaSession::SupportsRate(float Rate, bool Unthinned) const
+{
+	if (Rate == 1.0f)
+	{
+		return true;
+	}
+
+	if ((RateControl == NULL) || (RateSupport == NULL))
+	{
+		return false;
+	}
+
+	if (Unthinned)
+	{
+		return SUCCEEDED(RateSupport->IsRateSupported(FALSE, Rate, NULL));
+	}
+
+	return SUCCEEDED(RateSupport->IsRateSupported(TRUE, Rate, NULL));
+}
+
+
+bool FWmfMediaSession::SupportsScrubbing() const
+{
+	return CanScrub;
+}
+
+
+bool FWmfMediaSession::SupportsSeeking() const
+{
+	return (((Capabilities & MFSESSIONCAP_SEEK) != 0) && (Duration > FTimespan::Zero()));
+}
+
+
+/* FWmfMediaSession interface
+ *****************************************************************************/
+
+bool FWmfMediaSession::SetState(EMediaState NewState)
 {
 	RequestedState = NewState;
 
-	if ((NewState == EMediaStates::Closed) || (NewState == EMediaStates::Stopped))
+	if ((NewState == EMediaState::Closed) || (NewState == EMediaState::Stopped))
 	{
-		RequestedPosition = FTimespan::Zero();
+		RequestedTime = FTimespan::Zero();
 		RequestedRate = 0.0f;
 	}
 
@@ -187,10 +275,12 @@ STDMETHODIMP FWmfMediaSession::GetParameters(unsigned long*, unsigned long*)
 
 STDMETHODIMP FWmfMediaSession::Invoke(IMFAsyncResult* AsyncResult)
 {
+	// get event
 	TComPtr<IMFMediaEvent> Event;
 	
 	if (FAILED(MediaSession->EndGetEvent(AsyncResult, &Event)))
 	{
+		UE_LOG(LogWmfMedia, Verbose, TEXT("Failed to get event"));
 		return S_OK;
 	}
 
@@ -198,105 +288,104 @@ STDMETHODIMP FWmfMediaSession::Invoke(IMFAsyncResult* AsyncResult)
 	
 	if (FAILED(Event->GetType(&EventType)))
 	{
+		UE_LOG(LogWmfMedia, Verbose, TEXT("Failed to get session event type"));
 		return S_OK;
 	}
 
-	HRESULT EventResult;
-	Event->GetStatus(&EventResult);
+	// process event
+	HRESULT EventResult = S_OK;
+
+	if (FAILED(Event->GetStatus(&EventResult)))
+	{
+		UE_LOG(LogWmfMedia, Verbose, TEXT("Failed to get event result"));
+	}
+
+	UE_LOG(LogWmfMedia, Verbose, TEXT("Media session event: %s: %s"), *WmfMedia::SessionEventToString(EventType), *WmfMedia::ResultToString(EventResult));
 
 	FScopeLock ScopeLock(&CriticalSection);
 	{
-		if (EventType == MESessionClosed)
-		{
-			UpdateState(EMediaStates::Closed);
-			Capabilities = 0;
-		}
-		else if (EventType == MEError)
+		switch (EventType)
 		{
-			ErrorEvent.Broadcast(EventResult);
-			UpdateState(EMediaStates::Error);
-		}
-		else
-		{
-			switch (EventType)
-			{
-			case MEEndOfPresentation:
-				if (Looping && (CurrentRate > 0.0f))
-				{
-					RequestedPosition = FTimespan::Zero();
-					ChangeState();
-				}
-				break;
+		case MEBufferingStarted:
+			Buffering = true;
+			break;
 
-			case MESessionCapabilitiesChanged:
-				Capabilities = ::MFGetAttributeUINT32(Event, MF_EVENT_SESSIONCAPS, Capabilities);
-				break;
+		case MEBufferingStopped:
+			Buffering = false;
+			break;
 
-			case MESessionEnded:
-				if (Looping)
-				{
-					if (CurrentRate < 0.0f)
-					{
-						RequestedPosition = Duration;
-						ChangeState();
-					}
-				}
-				else
-				{
-					RequestedState = EMediaStates::Stopped;
-					UpdateState(EMediaStates::Stopped);
-				}
-				break;
+		case MEError:
+			UE_LOG(LogWmfMedia, Error, TEXT("An error occured in the media session: %s"), *WmfMedia::ResultToString(EventResult));
+			UpdateState(EMediaState::Error);
+			break;
 
-			case MESessionPaused:
-				UpdateState(EMediaStates::Paused);
-				break;
+		case MESessionCapabilitiesChanged:
+			Capabilities = ::MFGetAttributeUINT32(Event, MF_EVENT_SESSIONCAPS, Capabilities);
+			break;
 
-			case MESessionRateChanged:
-				// recover active playback rate if rate change failed
-				if (FAILED(EventResult) && (CurrentRate == RequestedRate))
-				{
-					PROPVARIANT Value;
-					PropVariantInit(&Value);
+		case MESessionClosed:
+			UpdateState(EMediaState::Closed);
+			Capabilities = 0;
+			break;
 
-					if (SUCCEEDED(Event->GetValue(&Value)) && (Value.vt == VT_R4))
-					{
-						CurrentRate = Value.fltVal;
-					}
+		case MESessionEnded:
+			if (Looping)
+			{
+				RequestedTime = (CurrentRate < 0.0f) ? Duration : FTimespan::Zero();
+				ChangeState();
+			}
+			else
+			{
+				RequestedState = EMediaState::Stopped;
+				UpdateState(EMediaState::Stopped);
+			}
+			break;
 
-					RequestedRate = CurrentRate;
-				}
+		case MESessionPaused:
+			UpdateState(EMediaState::Paused);
+			break;
 
-			case MESessionScrubSampleComplete:
-				if (GetInternalPosition() == RequestedPosition)
-				{
-					RequestedPosition = FTimespan::MinValue();
-				}
-				break;
+		case MESessionRateChanged:
+			// recover active playback rate if rate change failed
+			if (FAILED(EventResult) && (CurrentRate == RequestedRate))
+			{
+				PROPVARIANT Value;
+				PropVariantInit(&Value);
 
-			case MESessionStarted:
-				if (GetInternalPosition() == RequestedPosition)
+				if (SUCCEEDED(Event->GetValue(&Value)) && (Value.vt == VT_R4))
 				{
-					RequestedPosition = FTimespan::MinValue();
+					CurrentRate = Value.fltVal;
 				}
 
-				UpdateState(EMediaStates::Playing);
-				break;
-
-			case MESessionStopped:
-				UpdateState(EMediaStates::Stopped);
-				break;
+				RequestedRate = CurrentRate;
 			}
 
-			// bubble up event to player
-			SessionEvent.Broadcast(EventType);
-
-			// request the next event
-			if (FAILED(MediaSession->BeginGetEvent(this, NULL)))
+		case MESessionScrubSampleComplete:
+			if (GetInternalPosition() == RequestedTime)
 			{
-				Capabilities = 0;
-				CurrentState = EMediaStates::Error;
+				RequestedTime = FTimespan::MinValue();
 			}
+			break;
+
+		case MESessionStarted:
+			UpdateState(EMediaState::Playing);
+			break;
+
+		case MESessionStopped:
+			UpdateState(EMediaState::Stopped);
+			break;
+		}
+
+		// bubble up event to player
+		SessionEvent.Broadcast(EventType);
+
+		// request the next event
+		if (FAILED(MediaSession->BeginGetEvent(this, NULL)))
+		{
+			UE_LOG(LogWmfMedia, Verbose, TEXT("Failed to request next session event"));
+
+			Capabilities = 0;
+			CurrentState = EMediaState::Error;
 		}
 	}
 	
@@ -348,15 +437,16 @@ STDMETHODIMP_(ULONG) FWmfMediaSession::Release()
 bool FWmfMediaSession::ChangeState()
 {
 	// disallow state changes if session was closed or had an error
-	if (CurrentState == EMediaStates::Closed)
+	if (CurrentState == EMediaState::Closed)
 	{
 		return false;
 	}
 
 	FScopeLock ScopeLock(&CriticalSection);
 	{
-		if (CurrentState == EMediaStates::Error)
+		if (CurrentState == EMediaState::Error)
 		{
+			UE_LOG(LogWmfMedia, Verbose, TEXT("Closing media session on error"));
 			MediaSession->Close();
 
 			return true;
@@ -371,60 +461,81 @@ bool FWmfMediaSession::ChangeState()
 		}
 
 		// close session
-		if (RequestedState == EMediaStates::Closed)
+		if (RequestedState == EMediaState::Closed)
 		{
+			UE_LOG(LogWmfMedia, Verbose, TEXT("Closing media session as requested"));
 			MediaSession->Close();
 
 			return true;
 		}
 
 		// handle rate changes
-		bool bUpdateRateStopped = false;
-		UpdateRate(bUpdateRateStopped);
-		if (bUpdateRateStopped)
+		if (UpdateRate())
 		{
 			return true;
 		}
 
-		// handle state changes
-		bool StateChangeFailed = false;
-
-		if (RequestedState == EMediaStates::Paused)
+		// seek or scrub
+		if (RequestedTime >= FTimespan::Zero())
 		{
-			if (CurrentState != EMediaStates::Paused)
+			UE_LOG(LogWmfMedia, Verbose, TEXT("Starting playback to seek to %s"), *RequestedTime.ToString());
+
+			PROPVARIANT StartPosition;
 			{
-				StateChangePending = true;
-				StateChangeFailed = FAILED(MediaSession->Pause());
-				RequestedPosition = GetPosition();
+				StartPosition.vt = VT_I8;
+				StartPosition.hVal.QuadPart = RequestedTime.GetTicks();
+			}
+
+			HRESULT Result = MediaSession->Start(NULL, &StartPosition);
+
+			if (FAILED(Result))
+			{
+				UE_LOG(LogWmfMedia, Warning, TEXT("Failed to seek %i"), *WmfMedia::ResultToString(Result));
 			}
 		}
-		else if (RequestedState == EMediaStates::Playing)
+
+		// handle state changes
+		if (RequestedState == CurrentState)
 		{
-			if ((CurrentState != EMediaStates::Playing) || (RequestedPosition >= FTimespan::Zero()))
-			{
-				PROPVARIANT StartPosition;
+			return true;
+		}
 
-				if (RequestedPosition >= FTimespan::Zero())
-				{
-					StartPosition.vt = VT_I8;
-					StartPosition.hVal.QuadPart = RequestedPosition.GetTicks();
-				}
-				else
-				{
-					PropVariantInit(&StartPosition);
-				}
+		bool StateChangeFailed = false;
 
-				StateChangePending = true;
-				StateChangeFailed = FAILED(MediaSession->Start(NULL, &StartPosition));
-			}
+		if (RequestedState == EMediaState::Paused)
+		{
+			UE_LOG(LogWmfMedia, Verbose, TEXT("Pausing media session as requested"));
+
+			StateChangePending = true;
+			StateChangeFailed = FAILED(MediaSession->Pause());
+//			RequestedTime = GetTime();
 		}
-		else if (RequestedState == EMediaStates::Stopped)
+		else if (RequestedState == EMediaState::Playing)
 		{
-			if (CurrentState != EMediaStates::Stopped)
+			PROPVARIANT StartPosition;
+
+			if (RequestedTime >= FTimespan::Zero())
 			{
-				StateChangePending = true;
-				StateChangeFailed = FAILED(MediaSession->Stop());
+				UE_LOG(LogWmfMedia, Verbose, TEXT("Starting playback at %s"), *RequestedTime.ToString());
+
+				StartPosition.vt = VT_I8;
+				StartPosition.hVal.QuadPart = RequestedTime.GetTicks();
+			}
+			else
+			{
+				UE_LOG(LogWmfMedia, Verbose, TEXT("Starting playback at beginning"));
+				PropVariantInit(&StartPosition);
 			}
+
+			StateChangePending = true;
+			StateChangeFailed = FAILED(MediaSession->Start(NULL, &StartPosition));
+		}
+		else if (RequestedState == EMediaState::Stopped)
+		{
+			UE_LOG(LogWmfMedia, Verbose, TEXT("Stopping playback as requested"));
+
+			StateChangePending = true;
+			StateChangeFailed = FAILED(MediaSession->Stop());
 		}
 
 		if (StateChangeFailed)
@@ -439,44 +550,38 @@ bool FWmfMediaSession::ChangeState()
 }
 
 
-void FWmfMediaSession::UpdateRate(bool &OutStopped)
+bool FWmfMediaSession::UpdateRate()
 {
-	OutStopped = false;
+	bool Result = false;
 
 	if (RequestedRate != CurrentRate)
 	{
-		if (CurrentState != EMediaStates::Stopped)
+		if (CurrentState != EMediaState::Stopped)
 		{
-			if (RequestedPosition == FTimespan::MinValue())
+			UE_LOG(LogWmfMedia, Verbose, TEXT("Stopping playback for rate change"));
+
+			if (RequestedTime == FTimespan::MinValue())
 			{
-				RequestedPosition = GetPosition();
+				RequestedTime = GetTime();
 			}
 
 			MediaSession->Stop();
-
-			OutStopped = true;
+			Result = true;
 		}
 
-		bool RateChangeFailed;
-
-		if (IsRateSupported(RequestedRate, true))
+		if (SUCCEEDED(RateControl->SetRate(SupportsRate(RequestedRate, true) ? FALSE : TRUE, RequestedRate)))
 		{
-			RateChangeFailed = FAILED(RateControl->SetRate(FALSE, RequestedRate));
+			UE_LOG(LogWmfMedia, Verbose, TEXT("Set rate from %f to %f"), CurrentRate, RequestedRate);
+			CurrentRate = RequestedRate;
 		}
 		else
 		{
-			RateChangeFailed = FAILED(RateControl->SetRate(TRUE, RequestedRate));
-		}
-
-		if (RateChangeFailed)
-		{
+			UE_LOG(LogWmfMedia, Verbose, TEXT("Failed to set rate from %f to %f"), CurrentRate, RequestedRate);
 			RequestedRate = CurrentRate;
 		}
-		else
-		{
-			CurrentRate = RequestedRate;
-		}
 	}
+
+	return Result;
 }
 
 
@@ -494,23 +599,21 @@ FTimespan FWmfMediaSession::GetInternalPosition() const
 }
 
 
-void FWmfMediaSession::UpdateState(EMediaStates CompletedState)
+void FWmfMediaSession::UpdateState(EMediaState CompletedState)
 {
 	CurrentState = CompletedState;
 
-	if (CompletedState == EMediaStates::Playing)
+	if (CompletedState == EMediaState::Playing)
 	{
-		RequestedPosition = FTimespan::MinValue();
+		RequestedTime = FTimespan::MinValue();
 	}
 
 	if (CompletedState == RequestedState)
 	{
 		StateChangePending = false;
+		UpdateRate();
 
-		bool bUpdateRateStopped = false;
-		UpdateRate(bUpdateRateStopped);
-
-		if (RequestedState == EMediaStates::Stopped)
+		if (RequestedState == EMediaState::Stopped)
 		{
 			CurrentRate = 0.0f;
 			RequestedRate = 0.0f;
diff --git a/Engine/Plugins/Media/WmfMedia/Source/WmfMedia/Private/Wmf/WmfMediaSession.h b/Engine/Plugins/Media/WmfMedia/Source/WmfMedia/Private/Wmf/WmfMediaSession.h
--- a/Engine/Plugins/Media/WmfMedia/Source/WmfMedia/Private/Wmf/WmfMediaSession.h
+++ b/Engine/Plugins/Media/WmfMedia/Source/WmfMedia/Private/Wmf/WmfMediaSession.h
@@ -2,6 +2,7 @@
 
 #pragma once
 
+#include "IMediaControls.h"
 #include "AllowWindowsPlatformTypes.h"
 
 
@@ -9,28 +10,6 @@
 enum class EMediaPlaybackDirections;
 
 
-/**
- * Enumerates possible states of media playback.
- */
-enum class EMediaStates
-{
-	/** Video has been closed and cannot be played again. */
-	Closed,
-
-	/** Unrecoverable error occurred during playback. */
-	Error,
-
-	/** Playback has been paused, but can be resumed. */
-	Paused,
-
-	/** Video is currently playing. */
-	Playing,
-
-	/** Playback has been stopped, but can be restarted. */
-	Stopped
-};
-
-
 /**
  * Implements a media session that for handling asynchronous commands and callbacks.
  *
@@ -45,9 +24,13 @@ enum class EMediaStates
  */
 class FWmfMediaSession
 	: public IMFAsyncCallback
+	, public IMediaControls
 {
 public:
 
+	/** Default constructor. */
+	FWmfMediaSession();
+
 	/**
 	 * Creates and initializes a new instance.
 	 *
@@ -59,6 +42,23 @@ public:
 	/** Virtual destructor. */
 	virtual ~FWmfMediaSession() { }
 
+public:
+
+	//~ IMediaControls interface
+
+	virtual FTimespan GetDuration() const override;
+	virtual float GetRate() const override;
+	virtual EMediaState GetState() const override;
+	virtual TRange<float> GetSupportedRates(EMediaPlaybackDirections Direction, bool Unthinned) const override;
+	virtual FTimespan GetTime() const override;
+	virtual bool IsLooping() const override;
+	virtual bool Seek(const FTimespan& Time) override;
+	virtual bool SetLooping(bool Looping) override;
+	virtual bool SetRate(float Rate) override;
+	virtual bool SupportsRate(float Rate, bool Unthinned) const override;
+	virtual bool SupportsScrubbing() const override;
+	virtual bool SupportsSeeking() const override;
+
 public:
 
 	/**
@@ -71,114 +71,16 @@ public:
 		return Capabilities;
 	}
 
-	/**
-	 * Gets the current playback position.
-	 *
-	 * @return Presentation clock time, or FTimespan::MinValue if no clock is available.
-	 */
-	FTimespan GetPosition() const;
-
-	/**
-	 * Gets the current playback rate.
-	 *
-	 * @return Playback rate.
-	 */
-	float GetRate() const
-	{
-		return CurrentRate;
-	}
-
-	/**
-	 * Gets the state of this session.
-	 *
-	 * @return Media session state.
-	 */
-	EMediaStates GetState() const
-	{
-		return CurrentState;
-	}
-
-	/**
-	 * Gets the supported playback rate for the specified rate type.
-	 *
-	 * @param Direction The playback direction.
-	 * @param Unthinned Whether the rates are for unthinned playback.
-	 * @return The supported playback rate.
-	 */
-	TRange<float> GetSupportedRates(EMediaPlaybackDirections Direction, bool Unthinned) const;
-
-	/**
-	 * Checks whether playback is currently looping.
-	 *
-	 * @return true if looping, false otherwise.
-	 */
-	bool IsLooping() const
-	{
-		return Looping;
-	}
-
-	/**
-	 * Checks whether the specified playback rate is supported.
-	 *
-	 * @param Rate The rate to check (can be negative for reverse play).
-	 * @param Unthinned Whether no frames should be dropped at the given rate.
-	 * @return true if the rate is supported, false otherwise.
-	 */
-	bool IsRateSupported(float Rate, bool Unthinned) const;
-
-	/**
-	 * Sets whether playback should loop back to the beginning.
-	 *
-	 * @param InLooping Whether playback should be looped.
-	 */
-	void SetLooping(bool InLooping)
-	{
-		Looping = InLooping;
-	}
-
-	/**
-	 * Changes the playback position of the media.
-	 *
-	 * @param Position The position to set.
-	 * @return true if the position will be changed, false otherwise.
-	 */
-	bool SetPosition(const FTimespan& Position);
-
-	/**
-	 * Sets the playback rate.
-	 *
-	 * @param Rate The rate to set.
-	 * @return true if the rate will be changed, false otherwise.
-	 */
-	bool SetRate(float Rate);
-
 	/**
 	 * Sets the media state.
 	 *
 	 * @param NewState The media state to set.
 	 * @return true if the state will be changed, false otherwise.
 	 */
-	bool SetState(EMediaStates NewState);
-
-	/**
-	 * Checks whether this session supports scrubbing.
-	 *
-	 * @return true if scrubbing is supported, false otherwise.
-	 */
-	bool SupportsScrubbing() const
-	{
-		return CanScrub;
-	}
+	bool SetState(EMediaState NewState);
 
 public:
 
-	/** Gets an event delegate that is invoked when an error occured. */
-	DECLARE_EVENT_OneParam(FWmfMediaSession, FOnError, HRESULT /*Error*/)
-	FOnError& OnError()
-	{
-		return ErrorEvent;
-	}
-
 	/** Gets an event delegate that is invoked when the last presentation segment finished playing. */
 	DECLARE_EVENT_OneParam(FWmfMediaSession, FOnSessionEvent, MediaEventType /*EventType*/)
 	FOnSessionEvent& OnSessionEvent()
@@ -188,7 +90,7 @@ public:
 
 public:
 
-	// IMFAsyncCallback interface
+	//~ IMFAsyncCallback interface
 
 	STDMETHODIMP_(ULONG) AddRef();
 	STDMETHODIMP GetParameters(unsigned long*, unsigned long*);
@@ -206,11 +108,11 @@ protected:
 	bool ChangeState();
 
 	/**
-	 * Updates the CurrentPlayRate if it is different than the RequestedPlayRate
+	 * Updates the CurrentPlayRate if it is different than the RequestedPlayRate.
 	 *
-	 * @param OutStopped If true, the media player was stopped as a result of the change
+	 * @return Whether playback was stopped as a result of the update.
 	 */
-	void UpdateRate(bool &OutStopped);
+	bool UpdateRate();
 
 	/**
 	 * Gets the playback position as known by the internal WMF clock.
@@ -224,10 +126,13 @@ protected:
 	 *
 	 * @param CompletedState The state that was completed.
 	 */
-	void UpdateState(EMediaStates CompletedState);
+	void UpdateState(EMediaState CompletedState);
 
 private:
 
+	/** Whether the media source is buffering data. */
+	bool Buffering;
+
 	/** Caches whether the media supports scrubbing. */
 	bool CanScrub;
 
@@ -244,7 +149,7 @@ private:
 	float CurrentRate;
 
 	/** The current playback state. */
-	EMediaStates CurrentState;
+	EMediaState CurrentState;
 
 	/** The duration of the media. */
 	FTimespan Duration;
@@ -267,23 +172,20 @@ private:
 	/** Holds a reference counter for this instance. */
 	int32 RefCount;
 
-	/** The requested playback position. */
-	FTimespan RequestedPosition;
-
 	/** The requested playback rate. */
 	float RequestedRate;
 
+	/** The requested playback position. */
+	FTimespan RequestedTime;
+
 	/** The requested playback state. */
-	EMediaStates RequestedState;
+	EMediaState RequestedState;
 
 	/** Whether a state change is currently in progress. */
 	bool StateChangePending;
 
 private:
 
-	/** Holds an event delegate that is invoked when an error occurred. */
-	FOnError ErrorEvent;
-
 	/** Holds an event delegate that is invoked when an event occurs in the session. */
 	FOnSessionEvent SessionEvent;
 };
diff --git a/Engine/Plugins/Media/WmfMedia/Source/WmfMedia/Private/WmfMediaClasses.cpp b/Engine/Plugins/Media/WmfMedia/Source/WmfMedia/Private/WmfMediaClasses.cpp
new file mode 100644
--- /dev/null
+++ b/Engine/Plugins/Media/WmfMedia/Source/WmfMedia/Private/WmfMediaClasses.cpp
@@ -0,0 +1,9 @@
+// Copyright 1998-2016 Epic Games, Inc. All Rights Reserved.
+
+#include "WmfMediaPCH.h"
+#include "WmfMediaSettings.h"
+
+
+UWmfMediaSettings::UWmfMediaSettings()
+	: AllowNonStandardCodecs(false)
+{ }
