Made Settings module an Editor-only dependency

diff --git a/Engine/Plugins/Media/WmfMedia/Source/WmfMedia/WmfMedia.Build.cs b/Engine/Plugins/Media/WmfMedia/Source/WmfMedia/WmfMedia.Build.cs
--- a/Engine/Plugins/Media/WmfMedia/Source/WmfMedia/WmfMedia.Build.cs
+++ b/Engine/Plugins/Media/WmfMedia/Source/WmfMedia/WmfMedia.Build.cs
@@ -9,7 +9,6 @@ namespace UnrealBuildTool.Rules
             DynamicallyLoadedModuleNames.AddRange(
                 new string[] {
                     "Media",
-					"Settings",
 				}
             );
 
@@ -24,7 +23,6 @@ namespace UnrealBuildTool.Rules
 			PrivateIncludePathModuleNames.AddRange(
 				new string[] {
                     "Media",
-					"Settings",
 				}
 			);
 
@@ -37,7 +35,22 @@ namespace UnrealBuildTool.Rules
 				}
 			);
 
-            if ((Target.Platform == UnrealTargetPlatform.Win64) || (Target.Platform == UnrealTargetPlatform.Win32))
+			if (Target.Type == TargetRules.TargetType.Editor)
+			{
+				DynamicallyLoadedModuleNames.AddRange(
+					new string[] {
+						"Settings",
+					}
+				);
+
+				PrivateIncludePathModuleNames.AddRange(
+					new string[] {
+						"Settings",
+					}
+				);
+			}
+
+			if ((Target.Platform == UnrealTargetPlatform.Win64) || (Target.Platform == UnrealTargetPlatform.Win32))
             {
                 PublicDelayLoadDLLs.Add("mf.dll");
                 PublicDelayLoadDLLs.Add("mfplat.dll");
