Implementing 3064024 from //UE4/Dev-Framework to //UE4/Main

UE-33749 Audio files are imported as File Media Source asset

Returning false for .wav files UWmfFileMediaSourceFactory::FactoryCanImport

diff --git a/Engine/Plugins/Media/WmfMedia/Source/WmfMediaEditor/Private/WmfFileMediaSourceFactory.cpp b/Engine/Plugins/Media/WmfMedia/Source/WmfMediaEditor/Private/WmfFileMediaSourceFactory.cpp
--- a/Engine/Plugins/Media/WmfMedia/Source/WmfMediaEditor/Private/WmfFileMediaSourceFactory.cpp
+++ b/Engine/Plugins/Media/WmfMedia/Source/WmfMediaEditor/Private/WmfFileMediaSourceFactory.cpp
@@ -39,6 +39,11 @@ UWmfFileMediaSourceFactory::UWmfFileMediaSourceFactory(const FObjectInitializer&
 
 bool UWmfFileMediaSourceFactory::FactoryCanImport(const FString& Filename)
 {
+	FString FileExtension = FPaths::GetExtension(Filename);
+	if (FileExtension.ToUpper() == FString("WAV"))
+	{
+		return false;
+	}
 	return true;
 }
 
