Merged CL# 3150741 from Dev-Sequencer

diff --git a/Engine/Plugins/Media/WmfMedia/Source/WmfMedia/Private/Wmf/WmfMediaSession.cpp b/Engine/Plugins/Media/WmfMedia/Source/WmfMedia/Private/Wmf/WmfMediaSession.cpp
--- a/Engine/Plugins/Media/WmfMedia/Source/WmfMedia/Private/Wmf/WmfMediaSession.cpp
+++ b/Engine/Plugins/Media/WmfMedia/Source/WmfMedia/Private/Wmf/WmfMediaSession.cpp
@@ -190,6 +190,11 @@ bool FWmfMediaSession::IsLooping() const
 
 bool FWmfMediaSession::Seek(const FTimespan& Time)
 {
+	if ((CurrentState == EMediaState::Closed) || (CurrentState == EMediaState::Error))
+	{
+		return false;
+	}
+
 	RequestedTime = Time;
 
 	return ChangeState();
@@ -242,6 +247,11 @@ bool FWmfMediaSession::SetRate(float Rate)
 
 bool FWmfMediaSession::SupportsRate(float Rate, bool Unthinned) const
 {
+	if ((CurrentState == EMediaState::Closed) || (CurrentState == EMediaState::Error))
+	{
+		return false;
+	}
+
 	if (Rate == 1.0f)
 	{
 		return true;
@@ -278,6 +288,11 @@ bool FWmfMediaSession::SupportsSeeking() const
 
 bool FWmfMediaSession::SetState(EMediaState NewState)
 {
+	if (MediaSession == NULL)
+	{
+		return false;
+	}
+
 	RequestedState = NewState;
 
 	if ((NewState == EMediaState::Closed) || (NewState == EMediaState::Stopped))
