Copying //UE4/Dev-Sequencer to //UE4/Main (Source: //UE4/Dev-Sequencer @ 3617642)

diff --git a/Engine/Plugins/Media/WmfMedia/Source/WmfMedia/Private/Player/WmfMediaTracks.cpp b/Engine/Plugins/Media/WmfMedia/Source/WmfMedia/Private/Player/WmfMediaTracks.cpp
--- a/Engine/Plugins/Media/WmfMedia/Source/WmfMedia/Private/Player/WmfMediaTracks.cpp
+++ b/Engine/Plugins/Media/WmfMedia/Source/WmfMedia/Private/Player/WmfMediaTracks.cpp
@@ -1,19 +1,27 @@
 // Copyright 1998-2017 Epic Games, Inc. All Rights Reserved.
 
 #include "WmfMediaTracks.h"
+#include "WmfMediaPrivate.h"
 
-#include "IMediaAudioSink.h"
-#include "IMediaBinarySink.h"
+#if WMFMEDIA_SUPPORTED_PLATFORM
+
+#include "IHeadMountedDisplayModule.h"
 #include "IMediaOptions.h"
-#include "IMediaOverlaySink.h"
-#include "IMediaTextureSink.h"
 #include "Misc/ScopeLock.h"
 #include "UObject/Class.h"
+
+#if WITH_ENGINE
+	#include "Engine/Engine.h"
+#endif
+
+#include "WmfMediaAudioSample.h"
+#include "WmfMediaBinarySample.h"
+#include "WmfMediaOverlaySample.h"
 #include "WmfMediaSampler.h"
+#include "WmfMediaSettings.h"
+#include "WmfMediaTextureSample.h"
 #include "WmfMediaUtils.h"
 
-#if WMFMEDIA_SUPPORTED_PLATFORM
-
 #include "AllowWindowsPlatformTypes.h"
 
 #define LOCTEXT_NAMESPACE "FWmfMediaTracks"
@@ -23,20 +31,26 @@
  *****************************************************************************/
 
 FWmfMediaTracks::FWmfMediaTracks()
-	: SelectedAudioTrack(INDEX_NONE)
+	: AudioSamplePool(new FWmfMediaAudioSamplePool)
+	, MediaSourceChanged(false)
+	, SelectedAudioTrack(INDEX_NONE)
 	, SelectedBinaryTrack(INDEX_NONE)
 	, SelectedCaptionTrack(INDEX_NONE)
 	, SelectedVideoTrack(INDEX_NONE)
-	, AudioSink(nullptr)
-	, MetadataSink(nullptr)
-	, OverlaySink(nullptr)
-	, VideoSink(nullptr)
+	, TopologyChanged(false)
+	, VideoSamplePool(new FWmfMediaTextureSamplePool)
 { }
 
 
 FWmfMediaTracks::~FWmfMediaTracks()
 {
-	Reset();
+	Shutdown();
+
+	delete AudioSamplePool;
+	AudioSamplePool = nullptr;
+
+	delete VideoSamplePool;
+	VideoSamplePool = nullptr;
 }
 
 
@@ -56,11 +70,10 @@ void FWmfMediaTracks::AppendStats(FString &OutStats) const
 	}
 	else
 	{
-		for (const FAudioTrack& AudioTrack : AudioTracks)
+		for (const FTrack& Track : AudioTracks)
 		{
-			OutStats += FString::Printf(TEXT("    %s\n"), *AudioTrack.DisplayName.ToString());
-			OutStats += FString::Printf(TEXT("        Channels: %i\n"), AudioTrack.NumChannels);
-			OutStats += FString::Printf(TEXT("        Sample Rate: %i\n\n"), AudioTrack.SampleRate);
+			OutStats += FString::Printf(TEXT("    %s\n"), *Track.DisplayName.ToString());
+			OutStats += TEXT("        Not implemented yet");
 		}
 	}
 
@@ -73,21 +86,44 @@ void FWmfMediaTracks::AppendStats(FString &OutStats) const
 	}
 	else
 	{
-		for (const FVideoTrack& VideoTrack : VideoTracks)
+		for (const FTrack& Track : VideoTracks)
 		{
-			OutStats += FString::Printf(TEXT("    %s\n"), *VideoTrack.DisplayName.ToString());
-			OutStats += FString::Printf(TEXT("        BitRate: %i\n"), VideoTrack.BitRate);
-			OutStats += FString::Printf(TEXT("        Frame Rate: %i\n\n"), VideoTrack.FrameRate);
+			OutStats += FString::Printf(TEXT("    %s\n"), *Track.DisplayName.ToString());
+			OutStats += TEXT("        Not implemented yet");
 		}
 	}
 }
 
 
-TComPtr<IMFTopology> FWmfMediaTracks::CreateTopology() const
+void FWmfMediaTracks::ClearFlags()
+{
+	FScopeLock Lock(&CriticalSection);
+
+	MediaSourceChanged = false;
+	TopologyChanged = false;
+}
+
+
+TComPtr<IMFTopology> FWmfMediaTracks::CreateTopology()
 {
+	// validate streams
+	if (MediaSource == NULL)
+	{
+		return NULL; // nothing to play
+	}
+
+	if ((SelectedAudioTrack == INDEX_NONE) &&
+		(SelectedBinaryTrack == INDEX_NONE) &&
+		(SelectedCaptionTrack == INDEX_NONE) &&
+		(SelectedVideoTrack == INDEX_NONE))
+	{
+		return NULL; // no tracks selected
+	}
+
+	// create topology
 	TComPtr<IMFTopology> Topology;
 	{
-		HRESULT Result = ::MFCreateTopology(&Topology);
+		const HRESULT Result = ::MFCreateTopology(&Topology);
 
 		if (FAILED(Result))
 		{
@@ -96,233 +132,351 @@ TComPtr<IMFTopology> FWmfMediaTracks::CreateTopology() const
 		}
 	}
 
+	// add enabled streams to topology
+	bool TracksAdded = false;
+
 	if (AudioTracks.IsValidIndex(SelectedAudioTrack))
 	{
-		AddStreamToTopology(AudioTracks[SelectedAudioTrack], *Topology);
+		TracksAdded |= AddTrackToTopology(AudioTracks[SelectedAudioTrack], *Topology);
 	}
 
-	if (BinaryTracks.IsValidIndex(SelectedBinaryTrack))
+	if (MetadataTracks.IsValidIndex(SelectedBinaryTrack))
 	{
-		AddStreamToTopology(BinaryTracks[SelectedBinaryTrack], *Topology);
+		TracksAdded |= AddTrackToTopology(MetadataTracks[SelectedBinaryTrack], *Topology);
 	}
 
 	if (CaptionTracks.IsValidIndex(SelectedCaptionTrack))
 	{
-		AddStreamToTopology(CaptionTracks[SelectedCaptionTrack], *Topology);
+		TracksAdded |= AddTrackToTopology(CaptionTracks[SelectedCaptionTrack], *Topology);
 	}
 
 	if (VideoTracks.IsValidIndex(SelectedVideoTrack))
 	{
-		AddStreamToTopology(VideoTracks[SelectedVideoTrack], *Topology);
+		TracksAdded |= AddTrackToTopology(VideoTracks[SelectedVideoTrack], *Topology);
+	}
+
+	if (!TracksAdded)
+	{
+		return NULL;
 	}
 
 	return Topology;
 }
 
 
-void FWmfMediaTracks::FlushSinks()
+FTimespan FWmfMediaTracks::GetDuration() const
 {
 	FScopeLock Lock(&CriticalSection);
 
-	if (AudioSink != nullptr)
+	if (PresentationDescriptor == NULL)
 	{
-		AudioSink->FlushAudioSink();
+		return FTimespan::Zero();
 	}
+
+	UINT64 PresentationDuration = 0;
+	PresentationDescriptor->GetUINT64(MF_PD_DURATION, &PresentationDuration);
+	
+	return FTimespan(PresentationDuration);
+}
+
+
+void FWmfMediaTracks::GetFlags(bool& OutMediaSourceChanged, bool& OutTopologyChanged) const
+{
+	FScopeLock Lock(&CriticalSection);
+
+	OutMediaSourceChanged = MediaSourceChanged;
+	OutTopologyChanged = TopologyChanged;
 }
 
 
-bool FWmfMediaTracks::Initialize(IMFMediaSource& InMediaSource, IMFPresentationDescriptor& InPresentationDescriptor, FString& OutInfo)
+void FWmfMediaTracks::Initialize(IMFMediaSource* InMediaSource)
 {
-	Reset();
+	Shutdown();
 
 	FScopeLock Lock(&CriticalSection);
 
-	MediaSource = &InMediaSource;
-	PresentationDescriptor = &InPresentationDescriptor;
+	MediaSourceChanged = true;
 
-	// get number of streams
-	DWORD StreamCount = 0;
+	if (InMediaSource != NULL)
 	{
-		HRESULT Result = PresentationDescriptor->GetStreamDescriptorCount(&StreamCount);
+		// create presentation descriptor
+		TComPtr<IMFPresentationDescriptor> NewPresentationDescriptor;
+		{
+			const HRESULT Result = InMediaSource->CreatePresentationDescriptor(&NewPresentationDescriptor);
 
-		if (FAILED(Result))
+			if (FAILED(Result))
+			{
+				UE_LOG(LogWmfMedia, Error, TEXT("Failed to create presentation descriptor: %s"), *WmfMedia::ResultToString(Result));
+				return;
+			}
+		}
+
+		// get number of streams
+		DWORD StreamCount = 0;
 		{
-			UE_LOG(LogWmfMedia, Error, TEXT("Failed to get stream count (%s)"), *WmfMedia::ResultToString(Result));
-			return false;
+			const HRESULT Result = NewPresentationDescriptor->GetStreamDescriptorCount(&StreamCount);
+
+			if (FAILED(Result))
+			{
+				UE_LOG(LogWmfMedia, Error, TEXT("Failed to get stream count (%s)"), *WmfMedia::ResultToString(Result));
+				return;
+			}
+
+			UE_LOG(LogWmfMedia, Verbose, TEXT("Tracks %llx: Found %i streams"), this, StreamCount);
 		}
 
-		UE_LOG(LogWmfMedia, Verbose, TEXT("Found %i streams"), StreamCount);
-	}
+		// initialization successful
+		MediaSource = InMediaSource;
+		PresentationDescriptor = NewPresentationDescriptor;
 
-	// add streams (Media Foundation reports them in reverse order)
-	for (int32 StreamIndex = StreamCount - 1; StreamIndex >= 0; --StreamIndex)
-	{
-		AddStreamToTracks(StreamIndex, InMediaSource, InPresentationDescriptor, OutInfo);	
-		OutInfo += TEXT("\n");
+		// add streams (Media Foundation reports them in reverse order)
+		for (int32 StreamIndex = StreamCount - 1; StreamIndex >= 0; --StreamIndex)
+		{
+			AddStreamToTracks(StreamIndex, Info);
+			Info += TEXT("\n");
+		}
 	}
 
-	return true;
+	TopologyChanged = true;
 }
 
 
-void FWmfMediaTracks::Reset()
+void FWmfMediaTracks::Shutdown()
 {
 	FScopeLock Lock(&CriticalSection);
 
+	AudioSamplePool->Reset();
+	VideoSamplePool->Reset();
+
 	SelectedAudioTrack = INDEX_NONE;
 	SelectedBinaryTrack = INDEX_NONE;
 	SelectedCaptionTrack = INDEX_NONE;
 	SelectedVideoTrack = INDEX_NONE;
 
-	if (AudioSink != nullptr)
+	AudioTracks.Empty();
+	MetadataTracks.Empty();
+	CaptionTracks.Empty();
+	VideoTracks.Empty();
+
+	Info.Empty();
+
+	if (MediaSource != NULL)
 	{
-		AudioSink->ShutdownAudioSink();
+		MediaSource->Shutdown();
+		MediaSource = NULL;
 	}
 
-	if (MetadataSink != nullptr)
+	PresentationDescriptor.Reset();
+
+	MediaSourceChanged = false;
+	TopologyChanged = false;
+}
+
+
+/* IMediaSamples interface
+ *****************************************************************************/
+
+bool FWmfMediaTracks::FetchAudio(TRange<FTimespan> TimeRange, TSharedPtr<IMediaAudioSample, ESPMode::ThreadSafe>& OutSample)
+{
+	TSharedPtr<IMediaAudioSample, ESPMode::ThreadSafe> Sample;
+
+	if (!AudioSampleQueue.Peek(Sample))
 	{
-		MetadataSink->ShutdownBinarySink();
+		return false;
 	}
 
-	if (OverlaySink != nullptr)
+	const FTimespan SampleTime = Sample->GetTime();
+
+	if (!TimeRange.Overlaps(TRange<FTimespan>(SampleTime, SampleTime + Sample->GetDuration())))
 	{
-		OverlaySink->ShutdownOverlaySink();
+		return false;
 	}
 
-	if (VideoSink != nullptr)
+	if (!AudioSampleQueue.Dequeue(Sample))
 	{
-		VideoSink->ShutdownTextureSink();
+		return false;
 	}
 
-	AudioTracks.Empty();
-	BinaryTracks.Empty();
-	CaptionTracks.Empty();
-	VideoTracks.Empty();
+	OutSample = Sample;
 
-	MediaSource = nullptr;
-	PresentationDescriptor = nullptr;
+	return true;
 }
 
 
-void FWmfMediaTracks::SetPaused(bool Paused)
+bool FWmfMediaTracks::FetchCaption(TRange<FTimespan> TimeRange, TSharedPtr<IMediaOverlaySample, ESPMode::ThreadSafe>& OutSample)
 {
-	FScopeLock Lock(&CriticalSection);
+	TSharedPtr<IMediaOverlaySample, ESPMode::ThreadSafe> Sample;
 
-	if (AudioSink != nullptr)
+	if (!CaptionSampleQueue.Peek(Sample))
 	{
-		if (Paused)
-		{
-			AudioSink->PauseAudioSink();
-		}
-		else
-		{
-			AudioSink->ResumeAudioSink();
-		}
+		return false;
 	}
-}
 
+	const FTimespan SampleTime = Sample->GetTime();
 
-/* IMediaOutput interface
- *****************************************************************************/
+	if (!TimeRange.Overlaps(TRange<FTimespan>(SampleTime, SampleTime + Sample->GetDuration())))
+	{
+		return false;
+	}
 
-void FWmfMediaTracks::SetAudioSink(IMediaAudioSink* Sink)
-{
-	if (Sink != AudioSink)
+	if (!CaptionSampleQueue.Dequeue(Sample))
 	{
-		FScopeLock Lock(&CriticalSection);
+		return false;
+	}
 
-		if (AudioSink != nullptr)
-		{
-			AudioSink->ShutdownAudioSink();
-		}
+	OutSample = Sample;
 
-		AudioSink = Sink;	
-		InitializeAudioSink();
-	}
+	return true;
 }
 
 
-void FWmfMediaTracks::SetMetadataSink(IMediaBinarySink* Sink)
+bool FWmfMediaTracks::FetchMetadata(TRange<FTimespan> TimeRange, TSharedPtr<IMediaBinarySample, ESPMode::ThreadSafe>& OutSample)
 {
-	if (Sink != MetadataSink)
+	TSharedPtr<IMediaBinarySample, ESPMode::ThreadSafe> Sample;
+
+	if (!MetadataSampleQueue.Peek(Sample))
 	{
-		FScopeLock Lock(&CriticalSection);
+		return false;
+	}
 
-		if (MetadataSink != nullptr)
-		{
-			MetadataSink->ShutdownBinarySink();
-		}
+	const FTimespan SampleTime = Sample->GetTime();
 
-		MetadataSink = Sink;
-		InitializeMetadataSink();
+	if (!TimeRange.Overlaps(TRange<FTimespan>(SampleTime, SampleTime + Sample->GetDuration())))
+	{
+		return false;
 	}
-}
 
-
-void FWmfMediaTracks::SetOverlaySink(IMediaOverlaySink* Sink)
-{
-	if (Sink != OverlaySink)
+	if (!MetadataSampleQueue.Dequeue(Sample))
 	{
-		FScopeLock Lock(&CriticalSection);
+		return false;
+	}
 
-		if (OverlaySink != nullptr)
-		{
-			OverlaySink->ShutdownOverlaySink();
-		}
+	OutSample = Sample;
 
-		OverlaySink = Sink;
-		InitializeOverlaySink();
-	}
+	return true;
 }
 
 
-void FWmfMediaTracks::SetVideoSink(IMediaTextureSink* Sink)
+bool FWmfMediaTracks::FetchVideo(TRange<FTimespan> TimeRange, TSharedPtr<IMediaTextureSample, ESPMode::ThreadSafe>& OutSample)
 {
-	if (Sink != VideoSink)
+	TSharedPtr<IMediaTextureSample, ESPMode::ThreadSafe> Sample;
+
+	if (!VideoSampleQueue.Peek(Sample))
 	{
-		FScopeLock Lock(&CriticalSection);
+		return false;
+	}
 
-		if (VideoSink != nullptr)
-		{
-			VideoSink->ShutdownTextureSink();
-		}
+	const FTimespan SampleTime = Sample->GetTime();
+
+	if (!TimeRange.Overlaps(TRange<FTimespan>(SampleTime, SampleTime + Sample->GetDuration())))
+	{
+		return false;
+	}
 
-		VideoSink = Sink;
-		InitializeVideoSink();
+	if (!VideoSampleQueue.Dequeue(Sample))
+	{
+		return false;
 	}
+
+	OutSample = Sample;
+
+	return true;
+}
+
+
+void FWmfMediaTracks::FlushSamples()
+{
+	AudioSampleQueue.RequestFlush();
+	CaptionSampleQueue.RequestFlush();
+	MetadataSampleQueue.RequestFlush();
+	VideoSampleQueue.RequestFlush();
 }
 
 
 /* IMediaTracks interface
  *****************************************************************************/
 
-uint32 FWmfMediaTracks::GetAudioTrackChannels(int32 TrackIndex) const
+bool FWmfMediaTracks::GetAudioTrackFormat(int32 TrackIndex, int32 FormatIndex, FMediaAudioTrackFormat& OutFormat) const
 {
-	return AudioTracks.IsValidIndex(TrackIndex) ? AudioTracks[TrackIndex].NumChannels : 0;
-}
+	FScopeLock Lock(&CriticalSection);
+
+	const FFormat* Format = GetAudioFormat(TrackIndex, FormatIndex);
+	
+	if (Format == nullptr)
+	{
+		return false; // format not found
+	}
 
+	OutFormat.BitsPerSample = Format->Audio.BitsPerSample;
+	OutFormat.NumChannels = Format->Audio.NumChannels;
+	OutFormat.SampleRate = Format->Audio.SampleRate;
+	OutFormat.TypeName = Format->TypeName;
 
-uint32 FWmfMediaTracks::GetAudioTrackSampleRate(int32 TrackIndex) const
-{
-	return AudioTracks.IsValidIndex(TrackIndex) ? AudioTracks[TrackIndex].SampleRate : 0;
+	return true;
 }
 
 
 int32 FWmfMediaTracks::GetNumTracks(EMediaTrackType TrackType) const
 {
+	FScopeLock Lock(&CriticalSection);
+
 	switch (TrackType)
 	{
 	case EMediaTrackType::Audio:
 		return AudioTracks.Num();
+
 	case EMediaTrackType::Metadata:
-		return BinaryTracks.Num();
+		return MetadataTracks.Num();
+
 	case EMediaTrackType::Caption:
 		return CaptionTracks.Num();
+
 	case EMediaTrackType::Video:
 		return VideoTracks.Num();
+
+	default:
+		break; // unsupported track type
+	}
+
+	return 0;
+}
+
+
+int32 FWmfMediaTracks::GetNumTrackFormats(EMediaTrackType TrackType, int32 TrackIndex) const
+{
+	FScopeLock Lock(&CriticalSection);
+
+	switch (TrackType)
+	{
+	case EMediaTrackType::Audio:
+		if (AudioTracks.IsValidIndex(TrackIndex))
+		{
+			return AudioTracks[TrackIndex].Formats.Num();
+		}
+
+	case EMediaTrackType::Metadata:
+		if (MetadataTracks.IsValidIndex(TrackIndex))
+		{
+			return 1;
+		}
+
+	case EMediaTrackType::Caption:
+		if (CaptionTracks.IsValidIndex(TrackIndex))
+		{
+			return 1;
+		}
+
+	case EMediaTrackType::Video:
+		if (VideoTracks.IsValidIndex(TrackIndex))
+		{
+			return VideoTracks[TrackIndex].Formats.Num();
+		}
+
 	default:
-		return 0;
+		break; // unsupported track type
 	}
+
+	return 0;
 }
 
 
@@ -332,20 +486,28 @@ int32 FWmfMediaTracks::GetSelectedTrack(EMediaTrackType TrackType) const
 	{
 	case EMediaTrackType::Audio:
 		return SelectedAudioTrack;
+
 	case EMediaTrackType::Metadata:
 		return SelectedBinaryTrack;
+
 	case EMediaTrackType::Caption:
 		return SelectedCaptionTrack;
+
 	case EMediaTrackType::Video:
 		return SelectedVideoTrack;
+
 	default:
-		return INDEX_NONE;
+		break; // unsupported track type
 	}
+
+	return INDEX_NONE;
 }
 
 
 FText FWmfMediaTracks::GetTrackDisplayName(EMediaTrackType TrackType, int32 TrackIndex) const
 {
+	FScopeLock Lock(&CriticalSection);
+
 	switch (TrackType)
 	{
 	case EMediaTrackType::Audio:
@@ -356,9 +518,9 @@ FText FWmfMediaTracks::GetTrackDisplayName(EMediaTrackType TrackType, int32 Trac
 		break;
 	
 	case EMediaTrackType::Metadata:
-		if (BinaryTracks.IsValidIndex(TrackIndex))
+		if (MetadataTracks.IsValidIndex(TrackIndex))
 		{
-			return BinaryTracks[TrackIndex].DisplayName;
+			return MetadataTracks[TrackIndex].DisplayName;
 		}
 		break;
 
@@ -377,15 +539,26 @@ FText FWmfMediaTracks::GetTrackDisplayName(EMediaTrackType TrackType, int32 Trac
 		break;
 
 	default:
-		break;
+		break; // unsupported track type
 	}
 
 	return FText::GetEmpty();
 }
 
 
+int32 FWmfMediaTracks::GetTrackFormat(EMediaTrackType TrackType, int32 TrackIndex) const
+{
+	FScopeLock Lock(&CriticalSection);
+
+	const FTrack* Track = GetTrack(TrackType, TrackIndex);
+	return (Track != nullptr) ? Track->SelectedFormat : INDEX_NONE;
+}
+
+
 FString FWmfMediaTracks::GetTrackLanguage(EMediaTrackType TrackType, int32 TrackIndex) const
 {
+	FScopeLock Lock(&CriticalSection);
+
 	switch (TrackType)
 	{
 	case EMediaTrackType::Audio:
@@ -396,9 +569,9 @@ FString FWmfMediaTracks::GetTrackLanguage(EMediaTrackType TrackType, int32 Track
 		break;
 
 	case EMediaTrackType::Metadata:
-		if (BinaryTracks.IsValidIndex(TrackIndex))
+		if (MetadataTracks.IsValidIndex(TrackIndex))
 		{
-			return BinaryTracks[TrackIndex].Language;
+			return MetadataTracks[TrackIndex].Language;
 		}
 		break;
 
@@ -417,7 +590,7 @@ FString FWmfMediaTracks::GetTrackLanguage(EMediaTrackType TrackType, int32 Track
 		break;
 
 	default:
-		break;
+		break; // unsupported track type
 	}
 
 	return FString();
@@ -426,6 +599,8 @@ FString FWmfMediaTracks::GetTrackLanguage(EMediaTrackType TrackType, int32 Track
 
 FString FWmfMediaTracks::GetTrackName(EMediaTrackType TrackType, int32 TrackIndex) const
 {
+	FScopeLock Lock(&CriticalSection);
+
 	switch (TrackType)
 	{
 	case EMediaTrackType::Audio:
@@ -436,9 +611,9 @@ FString FWmfMediaTracks::GetTrackName(EMediaTrackType TrackType, int32 TrackInde
 		break;
 
 	case EMediaTrackType::Metadata:
-		if (BinaryTracks.IsValidIndex(TrackIndex))
+		if (MetadataTracks.IsValidIndex(TrackIndex))
 		{
-			return BinaryTracks[TrackIndex].Name;
+			return MetadataTracks[TrackIndex].Name;
 		}
 		break;
 
@@ -457,34 +632,36 @@ FString FWmfMediaTracks::GetTrackName(EMediaTrackType TrackType, int32 TrackInde
 		break;
 
 	default:
-		break;
+		break; // unsupported track type
 	}
 
 	return FString();
 }
 
 
-uint32 FWmfMediaTracks::GetVideoTrackBitRate(int32 TrackIndex) const
+bool FWmfMediaTracks::GetVideoTrackFormat(int32 TrackIndex, int32 FormatIndex, FMediaVideoTrackFormat& OutFormat) const
 {
-	return VideoTracks.IsValidIndex(TrackIndex) ? VideoTracks[TrackIndex].BitRate : 0;
-}
-
+	FScopeLock Lock(&CriticalSection);
 
-FIntPoint FWmfMediaTracks::GetVideoTrackDimensions(int32 TrackIndex) const
-{
-	return VideoTracks.IsValidIndex(TrackIndex) ? VideoTracks[TrackIndex].OutputDim : FIntPoint::ZeroValue;
-}
+	const FFormat* Format = GetVideoFormat(TrackIndex, FormatIndex);
+	
+	if (Format == nullptr)
+	{
+		return false; // format not found
+	}
 
+	OutFormat.Dim = Format->Video.OutputDim;
+	OutFormat.FrameRate = Format->Video.FrameRate;
+	OutFormat.FrameRates = Format->Video.FrameRates;
+	OutFormat.TypeName = Format->TypeName;
 
-float FWmfMediaTracks::GetVideoTrackFrameRate(int32 TrackIndex) const
-{
-	return VideoTracks.IsValidIndex(TrackIndex) ? VideoTracks[TrackIndex].FrameRate : 0.0f;
+	return true;
 }
 
 
 bool FWmfMediaTracks::SelectTrack(EMediaTrackType TrackType, int32 TrackIndex)
 {
-	if (PresentationDescriptor == nullptr)
+	if (PresentationDescriptor == NULL)
 	{
 		return false; // not initialized
 	}
@@ -494,196 +671,390 @@ bool FWmfMediaTracks::SelectTrack(EMediaTrackType TrackType, int32 TrackIndex)
 	switch (TrackType)
 	{
 	case EMediaTrackType::Audio:
-		if (TrackIndex != SelectedAudioTrack)
+		if (TrackIndex == SelectedAudioTrack)
 		{
-			UE_LOG(LogWmfMedia, Verbose, TEXT("Selecting audio track %i instead of %i (%i tracks)."), TrackIndex, SelectedAudioTrack, AudioTracks.Num());
+			return true; // already selected
+		}
 
-			if (SelectedAudioTrack != INDEX_NONE)
-			{
-				PresentationDescriptor->DeselectStream(AudioTracks[SelectedAudioTrack].StreamIndex);
-				SelectedAudioTrack = INDEX_NONE;
-			}
+		UE_LOG(LogWmfMedia, Verbose, TEXT("Tracks %llx: Selecting audio track %i instead of %i (%i tracks)"), this, TrackIndex, SelectedAudioTrack, AudioTracks.Num());
 
-			if ((TrackIndex == INDEX_NONE) || (AudioTracks.IsValidIndex(TrackIndex) && SUCCEEDED(PresentationDescriptor->SelectStream(AudioTracks[TrackIndex].StreamIndex))))
-			{
-				SelectedAudioTrack = TrackIndex;
-			}
+		// disable current track
+		if (SelectedAudioTrack != INDEX_NONE)
+		{
+			UE_LOG(LogWmfMedia, VeryVerbose, TEXT("Tracks %llx: Disabling audio track %i"), this, SelectedAudioTrack);
+
+			PresentationDescriptor->DeselectStream(AudioTracks[SelectedAudioTrack].StreamIndex);
+			SelectedAudioTrack = INDEX_NONE;
+		}
 
-			if (SelectedAudioTrack == TrackIndex)
+		// enable new track
+		if (TrackIndex != INDEX_NONE)
+		{
+			if (!AudioTracks.IsValidIndex(TrackIndex) || FAILED(PresentationDescriptor->SelectStream(AudioTracks[TrackIndex].StreamIndex)))
 			{
-				InitializeAudioSink();
+				return false; // invalid track index or failed to select
 			}
+
+			UE_LOG(LogWmfMedia, VeryVerbose, TEXT("Tracks %llx: Enabled audio track %i"), this, TrackIndex);
 		}
+
+		SelectedAudioTrack = TrackIndex;
 		break;
 
-	case EMediaTrackType::Metadata:
-		if (TrackIndex != SelectedBinaryTrack)
+	case EMediaTrackType::Caption:
+		if (TrackIndex == SelectedCaptionTrack)
 		{
-			UE_LOG(LogWmfMedia, Verbose, TEXT("Selecting binary track %i instead of %i (%i tracks)."), TrackIndex, SelectedBinaryTrack, BinaryTracks.Num());
+			return true;
+		}
 
-			if (SelectedBinaryTrack != INDEX_NONE)
-			{
-				PresentationDescriptor->DeselectStream(BinaryTracks[SelectedBinaryTrack].StreamIndex);
-				SelectedBinaryTrack = INDEX_NONE;
-			}
+		UE_LOG(LogWmfMedia, Verbose, TEXT("Tracks %llx: Selecting caption track %i instead of %i (%i tracks)"), this, TrackIndex, SelectedCaptionTrack, CaptionTracks.Num());
 
-			if ((TrackIndex == INDEX_NONE) || (BinaryTracks.IsValidIndex(TrackIndex) && SUCCEEDED(PresentationDescriptor->SelectStream(BinaryTracks[TrackIndex].StreamIndex))))
-			{
-				SelectedBinaryTrack = TrackIndex;
-			}
+		// disable current track
+		if (SelectedCaptionTrack != INDEX_NONE)
+		{
+			UE_LOG(LogWmfMedia, VeryVerbose, TEXT("Tracks %llx: Disabling caption track %i"), this, SelectedCaptionTrack);
+
+			PresentationDescriptor->DeselectStream(CaptionTracks[SelectedCaptionTrack].StreamIndex);
+			SelectedCaptionTrack = INDEX_NONE;
+		}
 
-			if (SelectedBinaryTrack == TrackIndex)
+		// enable new track
+		if (TrackIndex != INDEX_NONE)
+		{
+			if (!CaptionTracks.IsValidIndex(TrackIndex) || FAILED(PresentationDescriptor->SelectStream(CaptionTracks[TrackIndex].StreamIndex)))
 			{
-				InitializeMetadataSink();
+				return false; // invalid track index or failed to select
 			}
+
+			UE_LOG(LogWmfMedia, VeryVerbose, TEXT("Tracks %llx: Enabled caption track %i"), this, TrackIndex);
 		}
+
+		SelectedCaptionTrack = TrackIndex;
 		break;
 
-	case EMediaTrackType::Caption:
-		if (TrackIndex != SelectedCaptionTrack)
+	case EMediaTrackType::Metadata:
+		if (TrackIndex == SelectedBinaryTrack)
 		{
-			UE_LOG(LogWmfMedia, Verbose, TEXT("Selecting caption track %i instead of %i (%i tracks)."), TrackIndex, SelectedCaptionTrack, CaptionTracks.Num());
+			return true; // already selected
+		}
 
-			if (SelectedCaptionTrack != INDEX_NONE)
-			{
-				PresentationDescriptor->DeselectStream(CaptionTracks[SelectedCaptionTrack].StreamIndex);
-				SelectedCaptionTrack = INDEX_NONE;
-			}
+		UE_LOG(LogWmfMedia, Verbose, TEXT("Tracks %llx: Selecting binary track %i instead of %i (%i tracks)"), this, TrackIndex, SelectedBinaryTrack, MetadataTracks.Num());
 
-			if ((TrackIndex == INDEX_NONE) || (CaptionTracks.IsValidIndex(TrackIndex) && SUCCEEDED(PresentationDescriptor->SelectStream(CaptionTracks[TrackIndex].StreamIndex))))
-			{
-				SelectedCaptionTrack = TrackIndex;
-			}
+		// disable current track
+		if (SelectedBinaryTrack != INDEX_NONE)
+		{
+			UE_LOG(LogWmfMedia, VeryVerbose, TEXT("Tracks %llx: Disabling metadata track %i"), this, SelectedBinaryTrack);
+
+			PresentationDescriptor->DeselectStream(MetadataTracks[SelectedBinaryTrack].StreamIndex);
+			SelectedBinaryTrack = INDEX_NONE;
+		}
 
-			if (SelectedCaptionTrack == TrackIndex)
+		// enable new track
+		if (TrackIndex != INDEX_NONE)
+		{
+			if (!MetadataTracks.IsValidIndex(TrackIndex) || FAILED(PresentationDescriptor->SelectStream(MetadataTracks[TrackIndex].StreamIndex)))
 			{
-				InitializeOverlaySink();
+				return false; // invalid track index or failed to select
 			}
+
+			UE_LOG(LogWmfMedia, VeryVerbose, TEXT("Tracks %llx: Enabled metadata track %i"), this, TrackIndex);
 		}
+
+		SelectedBinaryTrack = TrackIndex;
 		break;
 
 	case EMediaTrackType::Video:
-		if (TrackIndex != SelectedVideoTrack)
+		if (TrackIndex == SelectedVideoTrack)
 		{
-			UE_LOG(LogWmfMedia, Verbose, TEXT("Selecting video track %i instead of %i (%i tracks)."), TrackIndex, SelectedVideoTrack, VideoTracks.Num());
+			return true; // already selected
+		}
 
-			if (SelectedVideoTrack != INDEX_NONE)
-			{
-				PresentationDescriptor->DeselectStream(VideoTracks[SelectedVideoTrack].StreamIndex);
-				SelectedVideoTrack = INDEX_NONE;
-			}
+		UE_LOG(LogWmfMedia, Verbose, TEXT("Tracks %llx: Selecting video track %i instead of %i (%i tracks)"), this, TrackIndex, SelectedVideoTrack, VideoTracks.Num());
 
-			if ((TrackIndex == INDEX_NONE) || (VideoTracks.IsValidIndex(TrackIndex) && SUCCEEDED(PresentationDescriptor->SelectStream(VideoTracks[TrackIndex].StreamIndex))))
-			{
-				SelectedVideoTrack = TrackIndex;
-			}
+		// disable current track
+		if (SelectedVideoTrack != INDEX_NONE)
+		{
+			UE_LOG(LogWmfMedia, VeryVerbose, TEXT("Tracks %llx: Disabling video track %i"), this, SelectedVideoTrack);
+
+			PresentationDescriptor->DeselectStream(VideoTracks[SelectedVideoTrack].StreamIndex);
+			SelectedVideoTrack = INDEX_NONE;
+		}
 
-			if (SelectedVideoTrack == TrackIndex)
+		// enable new track
+		if (TrackIndex != INDEX_NONE)
+		{
+			if (!VideoTracks.IsValidIndex(TrackIndex) || FAILED(PresentationDescriptor->SelectStream(VideoTracks[TrackIndex].StreamIndex)))
 			{
-				InitializeVideoSink();
+				return false; // invalid track index or failed to select
 			}
+
+			UE_LOG(LogWmfMedia, VeryVerbose, TEXT("Tracks %llx: Enabled video track %i"), this, TrackIndex);
 		}
+
+		SelectedVideoTrack = TrackIndex;
 		break;
 
 	default:
-		return false;
+		return false; // unsupported track type
 	}
 
-	SelectionChangedEvent.Broadcast();
+	TopologyChanged = true;
 
-	return false;
+	return true;
 }
 
 
-/* FWmfMediaTracks implementation
- *****************************************************************************/
-
-bool FWmfMediaTracks::AddStreamToTopology(const FStreamInfo& Stream, IMFTopology& Topology) const
+bool FWmfMediaTracks::SetTrackFormat(EMediaTrackType TrackType, int32 TrackIndex, int32 FormatIndex)
 {
-	GUID StreamType;
+	FScopeLock Lock(&CriticalSection);
+
+	FTrack* Track = nullptr;
+
+	switch (TrackType)
+	{
+	case EMediaTrackType::Audio:
+		if (AudioTracks.IsValidIndex(TrackIndex))
+		{
+			Track = &AudioTracks[TrackIndex];
+		}
+		break;
+
+	case EMediaTrackType::Caption:
+		if (CaptionTracks.IsValidIndex(TrackIndex))
+		{
+			Track = &CaptionTracks[TrackIndex];
+		}
+		break;
+
+	case EMediaTrackType::Metadata:
+		if (MetadataTracks.IsValidIndex(TrackIndex))
+		{
+			Track = &MetadataTracks[TrackIndex];
+		}
+		break;
+
+	case EMediaTrackType::Video:
+		if (VideoTracks.IsValidIndex(TrackIndex))
+		{
+			Track = &VideoTracks[TrackIndex];
+		}
+		break;
+
+	default:
+		break; // unsupported track type
+	};
+
+	if (Track == nullptr)
 	{
-		HRESULT Result = Stream.OutputType->GetGUID(MF_MT_MAJOR_TYPE, &StreamType);
+		return false;
+	}
+
+	if (Track->SelectedFormat != FormatIndex)
+	{
+		Track->SelectedFormat = FormatIndex;
+		TopologyChanged = true;
+	}
+
+	return true;
+}
+
+
+bool FWmfMediaTracks::SetVideoTrackFrameRate(int32 TrackIndex, int32 FormatIndex, float FrameRate)
+{
+	FScopeLock Lock(&CriticalSection);
+
+	const FFormat* Format = GetVideoFormat(TrackIndex, FormatIndex);
+
+	if (Format == nullptr)
+	{
+		return false; // format not found
+	}
+
+	if (Format->Video.FrameRate == FrameRate)
+	{
+		return true; // frame rate already set
+	}
+
+	if (!Format->Video.FrameRates.Contains(FrameRate))
+	{
+		return false; // frame rate not supported
+	}
+
+	int32 Numerator;
+	int32 Denominator;
+
+	if (!WmfMedia::FrameRateToRatio(FrameRate, Numerator, Denominator))
+	{
+		return false; // invalid frame rate
+	}
+
+	return SUCCEEDED(::MFSetAttributeRatio(Format->InputType, MF_MT_FRAME_RATE, Numerator, Denominator));
+}
+
+
+/* FWmfMediaTracks implementation
+ *****************************************************************************/
+
+bool FWmfMediaTracks::AddTrackToTopology(const FTrack& Track, IMFTopology& Topology) const
+{
+	// skip if no format selected
+	if (!Track.Formats.IsValidIndex(Track.SelectedFormat))
+	{
+		return false;
+	}
+
+	// get selected format
+	const FFormat& Format = Track.Formats[Track.SelectedFormat];
+
+	check(Format.InputType.IsValid());
+	check(Format.OutputType.IsValid());
+
+	GUID MajorType;
+	{
+		const HRESULT Result = Format.OutputType->GetGUID(MF_MT_MAJOR_TYPE, &MajorType);
 		check(SUCCEEDED(Result));
 	}
 
-	TComPtr<IMFActivate> OutputActivator;
+	// skip audio if necessary
+	if (MajorType == MFMediaType_Audio)
 	{
-		if ((StreamType == MFMediaType_Audio) && GetDefault<UWmfMediaSettings>()->NativeAudioOut)
+		if (::waveOutGetNumDevs() == 0)
 		{
-			HRESULT Result = MFCreateAudioRendererActivate(&OutputActivator);
+			return false; // no audio device
+		}
 
-			if (FAILED(Result))
-			{
-				UE_LOG(LogWmfMedia, Warning, TEXT("Failed to create audio renderer for stream %i (%s)"), Stream.StreamIndex, *WmfMedia::ResultToString(Result));
-				return false;
-			}
+#if WITH_ENGINE
+		if ((GEngine != nullptr) && !GEngine->UseSound())
+		{
+			return false; // audio disabled
+		}
+
+		if ((GEngine == nullptr) && !GetDefault<UWmfMediaSettings>()->NativeAudioOut)
+		{
+			return false; // no engine audio
+		}
+#else
+		if (!GetDefault<UWmfMediaSettings>()->NativeAudioOut)
+		{
+			return false; // native audio disabled
+		}
+#endif
+	}
+
+	// set input type
+	check(Track.Handler.IsValid());
+	{
+		const HRESULT Result = Track.Handler->SetCurrentMediaType(Format.InputType);
+
+		if (FAILED(Result))
+		{
+			UE_LOG(LogWmfMedia, Warning, TEXT("Failed to set current media type for stream %i (%s)"), Track.StreamIndex, *WmfMedia::ResultToString(Result));
+			return false;
+		}
+	}
+
+	// create output activator
+	TComPtr<IMFActivate> OutputActivator;
+
+	if ((MajorType == MFMediaType_Audio) && GetDefault<UWmfMediaSettings>()->NativeAudioOut)
+	{
+		// create native audio renderer
+		HRESULT Result = MFCreateAudioRendererActivate(&OutputActivator);
+
+		if (FAILED(Result))
+		{
+			UE_LOG(LogWmfMedia, Warning, TEXT("Failed to create audio renderer for stream %i (%s)"), Track.StreamIndex, *WmfMedia::ResultToString(Result));
+			return false;
+		}
 
 #if WITH_ENGINE
-			// allow HMD to override audio output device
-			if (IHeadMountedDisplayModule::IsAvailable())
+		// allow HMD to override audio output device
+		if (IHeadMountedDisplayModule::IsAvailable())
+		{
+			FString AudioOutputDevice = IHeadMountedDisplayModule::Get().GetAudioOutputDevice();
+
+			if (!AudioOutputDevice.IsEmpty())
 			{
-				FString AudioOutputDevice = IHeadMountedDisplayModule::Get().GetAudioOutputDevice();
+				Result = OutputActivator->SetString(MF_AUDIO_RENDERER_ATTRIBUTE_ENDPOINT_ID, *AudioOutputDevice);
 
-				if (!AudioOutputDevice.IsEmpty())
+				if (FAILED(Result))
 				{
-					Result = OutputActivator->SetString(MF_AUDIO_RENDERER_ATTRIBUTE_ENDPOINT_ID, *AudioOutputDevice);
-
-					if (FAILED(Result))
-					{
-						UE_LOG(LogWmfMedia, Warning, TEXT("Failed to override HMD audio output device for stream %i (%s)"), Stream.StreamIndex, *WmfMedia::ResultToString(Result));
-						return false;
-					}
+					UE_LOG(LogWmfMedia, Warning, TEXT("Failed to override HMD audio output device for stream %i (%s)"), Track.StreamIndex, *WmfMedia::ResultToString(Result));
+					return false;
 				}
 			}
+		}
 #endif //WITH_ENGINE
+	}
+	else
+	{
+		// create custom sampler
+		TComPtr<FWmfMediaSampler> Sampler = new FWmfMediaSampler();
+
+		if (MajorType == MFMediaType_Audio)
+		{
+			Sampler->OnClock().AddRaw(this, &FWmfMediaTracks::HandleMediaSamplerClock, EMediaTrackType::Audio);
+			Sampler->OnSample().AddRaw(this, &FWmfMediaTracks::HandleMediaSamplerAudioSample);
 		}
-		else
+		else if (MajorType == MFMediaType_SAMI)
+		{
+			Sampler->OnSample().AddRaw(this, &FWmfMediaTracks::HandleMediaSamplerCaptionSample);
+		}
+		else if (MajorType == MFMediaType_Binary)
+		{
+			Sampler->OnSample().AddRaw(this, &FWmfMediaTracks::HandleMediaSamplerMetadataSample);
+		}
+		else if ((MajorType == MFMediaType_Video) || (MajorType == MFMediaType_Image))
 		{
-			check(Stream.Sampler != NULL);
-			HRESULT Result = ::MFCreateSampleGrabberSinkActivate(Stream.OutputType, Stream.Sampler, &OutputActivator);
+			Sampler->OnSample().AddRaw(this, &FWmfMediaTracks::HandleMediaSamplerVideoSample);
+		}
 
-			if (FAILED(Result))
-			{
-				UE_LOG(LogWmfMedia, Warning, TEXT("Failed to create sampler grabber sink for stream %i (%s)"), Stream.StreamIndex, *WmfMedia::ResultToString(Result));
-				return false;
-			}
+		const HRESULT Result = ::MFCreateSampleGrabberSinkActivate(Format.OutputType, Sampler, &OutputActivator);
+
+		if (FAILED(Result))
+		{
+			UE_LOG(LogWmfMedia, Warning, TEXT("Failed to create sampler grabber sink for stream %i (%s)"), Track.StreamIndex, *WmfMedia::ResultToString(Result));
+			return false;
 		}
 	}
 
+	if (!OutputActivator.IsValid())
+	{
+		return false;
+	}
+
 	// set up output node
 	TComPtr<IMFTopologyNode> OutputNode;
+
+	if (FAILED(::MFCreateTopologyNode(MF_TOPOLOGY_OUTPUT_NODE, &OutputNode)) ||
+		FAILED(OutputNode->SetObject(OutputActivator)) ||
+		FAILED(OutputNode->SetUINT32(MF_TOPONODE_STREAMID, 0)) ||
+		FAILED(OutputNode->SetUINT32(MF_TOPONODE_NOSHUTDOWN_ON_REMOVE, FALSE)) ||
+		FAILED(Topology.AddNode(OutputNode)))
 	{
-		if (FAILED(::MFCreateTopologyNode(MF_TOPOLOGY_OUTPUT_NODE, &OutputNode)) ||
-			FAILED(OutputNode->SetObject(OutputActivator)) ||
-			FAILED(OutputNode->SetUINT32(MF_TOPONODE_STREAMID, 0)) ||
-			FAILED(OutputNode->SetUINT32(MF_TOPONODE_NOSHUTDOWN_ON_REMOVE, FALSE)) ||
-			FAILED(Topology.AddNode(OutputNode)))
-		{
-			UE_LOG(LogWmfMedia, Warning, TEXT("Failed to configure output node for stream %i"), Stream.StreamIndex);
-			return false;
-		}
+		UE_LOG(LogWmfMedia, Warning, TEXT("Failed to configure output node for stream %i"), Track.StreamIndex);
+		return false;
 	}
 
 	// set up source node
 	TComPtr<IMFTopologyNode> SourceNode;
+
+	if (FAILED(::MFCreateTopologyNode(MF_TOPOLOGY_SOURCESTREAM_NODE, &SourceNode)) ||
+		FAILED(SourceNode->SetUnknown(MF_TOPONODE_SOURCE, MediaSource)) ||
+		FAILED(SourceNode->SetUnknown(MF_TOPONODE_PRESENTATION_DESCRIPTOR, PresentationDescriptor)) ||
+		FAILED(SourceNode->SetUnknown(MF_TOPONODE_STREAM_DESCRIPTOR, Track.Descriptor)) ||
+		FAILED(Topology.AddNode(SourceNode)))
 	{
-		if (FAILED(::MFCreateTopologyNode(MF_TOPOLOGY_SOURCESTREAM_NODE, &SourceNode)) ||
-			FAILED(SourceNode->SetUnknown(MF_TOPONODE_SOURCE, MediaSource)) ||
-			FAILED(SourceNode->SetUnknown(MF_TOPONODE_PRESENTATION_DESCRIPTOR, PresentationDescriptor)) ||
-			FAILED(SourceNode->SetUnknown(MF_TOPONODE_STREAM_DESCRIPTOR, Stream.Descriptor)) ||
-			FAILED(Topology.AddNode(SourceNode)))
-		{
-			UE_LOG(LogWmfMedia, Warning, TEXT("Failed to configure source node for stream %i"), Stream.StreamIndex);
-			return false;
-		}
+		UE_LOG(LogWmfMedia, Warning, TEXT("Failed to configure source node for stream %i"), Track.StreamIndex);
+		return false;
 	}
 
 	// connect nodes
-	HRESULT Result = SourceNode->ConnectOutput(0, OutputNode, 0);
+	const HRESULT Result = SourceNode->ConnectOutput(0, OutputNode, 0);
 
 	if (FAILED(Result))
 	{
-		UE_LOG(LogWmfMedia, Warning, TEXT("Failed to connect topology nodes for stream %i (%s)"), Stream.StreamIndex, *WmfMedia::ResultToString(Result));
+		UE_LOG(LogWmfMedia, Warning, TEXT("Failed to connect topology nodes for stream %i (%s)"), Track.StreamIndex, *WmfMedia::ResultToString(Result));
 		return false;
 	}
 
@@ -691,16 +1062,15 @@ bool FWmfMediaTracks::AddStreamToTopology(const FStreamInfo& Stream, IMFTopology
 }
 
 
-void FWmfMediaTracks::AddStreamToTracks(uint32 StreamIndex, IMFMediaSource& InMediaSource, IMFPresentationDescriptor& InPresentationDescriptor, FString& OutInfo)
+void FWmfMediaTracks::AddStreamToTracks(uint32 StreamIndex, FString& OutInfo)
 {
-	int32 CurrentTrackIndex = -1;
 	OutInfo += FString::Printf(TEXT("Stream %i\n"), StreamIndex);
 
 	// get stream descriptor
 	TComPtr<IMFStreamDescriptor> StreamDescriptor;
 	{
 		BOOL Selected = FALSE;
-		HRESULT Result = PresentationDescriptor->GetStreamDescriptorByIndex(StreamIndex, &Selected, &StreamDescriptor);
+		const HRESULT Result = PresentationDescriptor->GetStreamDescriptorByIndex(StreamIndex, &Selected, &StreamDescriptor);
 
 		if (FAILED(Result))
 		{
@@ -719,7 +1089,7 @@ void FWmfMediaTracks::AddStreamToTracks(uint32 StreamIndex, IMFMediaSource& InMe
 	// get media type handler
 	TComPtr<IMFMediaTypeHandler> Handler;
 	{
-		HRESULT Result = StreamDescriptor->GetMediaTypeHandler(&Handler);
+		const HRESULT Result = StreamDescriptor->GetMediaTypeHandler(&Handler);
 
 		if (FAILED(Result))
 		{
@@ -733,7 +1103,7 @@ void FWmfMediaTracks::AddStreamToTracks(uint32 StreamIndex, IMFMediaSource& InMe
 	// skip unsupported handler types
 	GUID MajorType;
 	{
-		HRESULT Result = Handler->GetMajorType(&MajorType);
+		const HRESULT Result = Handler->GetMajorType(&MajorType);
 
 		if (FAILED(Result))
 		{
@@ -743,7 +1113,7 @@ void FWmfMediaTracks::AddStreamToTracks(uint32 StreamIndex, IMFMediaSource& InMe
 			return;
 		}
 
-		UE_LOG(LogWmfMedia, Verbose, TEXT("Major type of stream %i is %s"), StreamIndex, *WmfMedia::MajorTypeToString(MajorType));
+		UE_LOG(LogWmfMedia, Verbose, TEXT("Tracks %llx: Major type of stream %i is %s"), this, StreamIndex, *WmfMedia::MajorTypeToString(MajorType));
 		OutInfo += FString::Printf(TEXT("    Type: %s\n"), *WmfMedia::MajorTypeToString(MajorType));
 
 		if ((MajorType != MFMediaType_Audio) &&
@@ -757,552 +1127,610 @@ void FWmfMediaTracks::AddStreamToTracks(uint32 StreamIndex, IMFMediaSource& InMe
 		}
 	}
 
-	// get media type and make it current
-	TComPtr<IMFMediaType> MediaType;
+	// @todo gmp: handle protected content
+	const bool Protected = ::MFGetAttributeUINT32(StreamDescriptor, MF_SD_PROTECTED, FALSE) != 0;
 	{
-		DWORD NumSupportedTypes = 0;
-		HRESULT Result = Handler->GetMediaTypeCount(&NumSupportedTypes);
+		if (Protected)
+		{
+			OutInfo += FString::Printf(TEXT("    Protected content\n"));
+		}
+	}
+
+	// get number of track formats
+	DWORD NumMediaTypes = 0;
+	{
+		const HRESULT Result = Handler->GetMediaTypeCount(&NumMediaTypes);
 
 		if (FAILED(Result))
 		{
-			UE_LOG(LogWmfMedia, Warning, TEXT("Failed to get number of supported media types in stream %i of type %s"), StreamIndex, *WmfMedia::MajorTypeToString(MajorType));
-			OutInfo += TEXT("    failed to get supported media types\n");
+			UE_LOG(LogWmfMedia, Warning, TEXT("Failed to get number of track formats in stream %i of type %s"), StreamIndex, *WmfMedia::MajorTypeToString(MajorType));
+			OutInfo += TEXT("    failed to get track formats\n");
 
 			return;
 		}
+	}
+
+	// create & add track
+	FTrack* Track = nullptr;
 
-		if (NumSupportedTypes > 0)
+	if (MajorType == MFMediaType_Audio)
+	{
+		const int32 TrackIndex = AudioTracks.AddDefaulted();
+		Track = &AudioTracks[TrackIndex];
+	}
+	else if (MajorType == MFMediaType_SAMI)
+	{
+		const int32 TrackIndex = CaptionTracks.AddDefaulted();
+		Track = &CaptionTracks[TrackIndex];
+	}
+	else if (MajorType == MFMediaType_Binary)
+	{
+		const int32 TrackIndex = MetadataTracks.AddDefaulted();
+		Track = &MetadataTracks[TrackIndex];
+	}
+	else if ((MajorType == MFMediaType_Video) || (MajorType == MFMediaType_Image))
+	{
+		const int32 TrackIndex = VideoTracks.AddDefaulted();
+		Track = &VideoTracks[TrackIndex];
+	}
+
+	check(Track != nullptr);
+
+	// get current format
+	TComPtr<IMFMediaType> CurrentMediaType;
+	{
+		if (FAILED(Handler->GetCurrentMediaType(&CurrentMediaType)))
 		{
-			for (DWORD TypeIndex = 0; TypeIndex < NumSupportedTypes; ++TypeIndex)
-			{
-				if (SUCCEEDED(Handler->GetMediaTypeByIndex(TypeIndex, &MediaType)) &&
-					SUCCEEDED(Handler->SetCurrentMediaType(MediaType)))
-				{
-					break;
-				}
-			}
+			UE_LOG(LogWmfMedia, Warning, TEXT("Failed to get current media type in stream %i of type %s"), StreamIndex, *WmfMedia::MajorTypeToString(MajorType));
 		}
+	}
+
+	Track->SelectedFormat = INDEX_NONE;
+
+	// add track formats
+	const bool AllowNonStandardCodecs = GetDefault<UWmfMediaSettings>()->AllowNonStandardCodecs;
+
+	for (DWORD TypeIndex = 0; TypeIndex < NumMediaTypes; ++TypeIndex)
+	{
+		OutInfo += FString::Printf(TEXT("    Format %i\n"), TypeIndex);
+
+		// get media type
+		TComPtr<IMFMediaType> MediaType;
 
-		if (MediaType == NULL)
+		if (FAILED(Handler->GetMediaTypeByIndex(TypeIndex, &MediaType)))
 		{
-			UE_LOG(LogWmfMedia, Warning, TEXT("No supported media type in stream %i of type %s"), StreamIndex, *WmfMedia::MajorTypeToString(MajorType));
-			OutInfo += TEXT("    unsupported media type\n");
+			OutInfo += TEXT("        failed to get media type\n");
 
-			return;
+			continue;
 		}
-	}
 
-	// get sub-type
-	GUID SubType;
-	{
+		// get sub-type
+		GUID SubType;
+
 		if (MajorType == MFMediaType_SAMI)
 		{
 			FMemory::Memzero(SubType);
 		}
 		else
 		{
-			HRESULT Result = MediaType->GetGUID(MF_MT_SUBTYPE, &SubType);
+			const HRESULT Result = MediaType->GetGUID(MF_MT_SUBTYPE, &SubType);
 
 			if (FAILED(Result))
 			{
-				UE_LOG(LogWmfMedia, Warning, TEXT("Failed to get sub-type of stream %i (%s)"), StreamIndex, *WmfMedia::ResultToString(Result));
-				OutInfo += TEXT("    failed to get sub-type\n");
+				UE_LOG(LogWmfMedia, Warning, TEXT("Failed to get sub-type of format %i in stream %i (%s)"), TypeIndex, StreamIndex, *WmfMedia::ResultToString(Result));
+				OutInfo += TEXT("        failed to get sub-type\n");
 
-				return;
+				continue;;
 			}
-
-			UE_LOG(LogWmfMedia, Verbose, TEXT("Sub-type of stream %i is %s"), StreamIndex, *WmfMedia::SubTypeToString(SubType));
-
-			OutInfo += FString::Printf(TEXT("    Codec: %s\n"), *WmfMedia::SubTypeToString(SubType));
 		}
-	}
 
-	// @todo gmp: handle protected content
-	const bool Protected = ::MFGetAttributeUINT32(StreamDescriptor, MF_SD_PROTECTED, FALSE) != 0;
-	{
-		if (Protected)
-		{
-			OutInfo += FString::Printf(TEXT("    Protected content\n"));
-		}
-	}
+		const FString TypeName = WmfMedia::SubTypeToString(SubType);
+		OutInfo += FString::Printf(TEXT("        Codec: %s\n"), *TypeName);
 
-	// configure output type
-	TComPtr<IMFMediaType> OutputType;
-	{
-		HRESULT Result = ::MFCreateMediaType(&OutputType);
+		// create output type
+		TComPtr<IMFMediaType> OutputType = WmfMedia::CreateOutputType(MajorType, SubType, AllowNonStandardCodecs);
 
-		if (FAILED(Result))
+		if (!OutputType.IsValid())
 		{
-			UE_LOG(LogWmfMedia, Warning, TEXT("Failed to create output type for stream %i (%s)"), StreamIndex, *WmfMedia::ResultToString(Result));
-			OutInfo += TEXT("    failed to create output type\n");
+			OutInfo += TEXT("        failed to create output type\n");
 
-			return;
+			continue;
 		}
 
-		Result = OutputType->SetUINT32(MF_MT_ALL_SAMPLES_INDEPENDENT, TRUE);
+		// add format details
+		int32 FormatIndex = INDEX_NONE;
 
-		if (FAILED(Result))
+		if (MajorType == MFMediaType_Audio)
 		{
-			UE_LOG(LogWmfMedia, Warning, TEXT("Failed to initialize output type for stream %i (%s)"), StreamIndex, *WmfMedia::ResultToString(Result));
-			OutInfo += TEXT("    failed to initialize output type\n");
-
-			return;
+			const uint32 BitsPerSample = ::MFGetAttributeUINT32(MediaType, MF_MT_AUDIO_BITS_PER_SAMPLE, 16);
+			const uint32 NumChannels = ::MFGetAttributeUINT32(MediaType, MF_MT_AUDIO_NUM_CHANNELS, 0);
+			const uint32 SampleRate = ::MFGetAttributeUINT32(MediaType, MF_MT_AUDIO_SAMPLES_PER_SECOND, 0);
+
+			FormatIndex = Track->Formats.Add({
+				MediaType,
+				OutputType,
+				TypeName,
+				{
+					BitsPerSample,
+					NumChannels,
+					SampleRate
+				},
+				{ 0 }
+			});
+
+			OutInfo += FString::Printf(TEXT("        Channels: %i\n"), NumChannels);
+			OutInfo += FString::Printf(TEXT("        Sample Rate: %i Hz\n"), SampleRate);
+			OutInfo += FString::Printf(TEXT("        Bits Per Sample: %i\n"), BitsPerSample);
 		}
-
-		if (MajorType == MFMediaType_Audio)
+		else if (MajorType == MFMediaType_SAMI)
 		{
-			// filter unsupported audio formats
-			if (FMemory::Memcmp(&SubType.Data2, &MFMPEG4Format_Base.Data2, 12) == 0)
+			FormatIndex = Track->Formats.Add({
+				MediaType,
+				OutputType,
+				TypeName,
+				{ 0 },
+				{ 0 }
+			});
+		}
+		else if (MajorType == MFMediaType_Binary)
+		{
+			FormatIndex = Track->Formats.Add({
+				MediaType,
+				OutputType,
+				TypeName,
+				{ 0 },
+				{ 0 }
+			});
+		}
+		else if (MajorType == MFMediaType_Video)
+		{
+			GUID OutputSubType;
 			{
-				if (GetDefault<UWmfMediaSettings>()->AllowNonStandardCodecs)
-				{
-					UE_LOG(LogWmfMedia, Verbose, TEXT("Non-standard MP4 audio type '%s' (%s) for stream %i"), *WmfMedia::FourccToString(SubType.Data1), *WmfMedia::GuidToString(SubType), StreamIndex);
-					OutInfo += TEXT("    non-standard sub-type\n");
-				}
-				else
+				const HRESULT Result = OutputType->GetGUID(MF_MT_SUBTYPE, &OutputSubType);
+
+				if (FAILED(Result))
 				{
-					const bool DocumentedFormat =
-						(SubType.Data1 == WAVE_FORMAT_ADPCM) ||
-						(SubType.Data1 == WAVE_FORMAT_ALAW) ||
-						(SubType.Data1 == WAVE_FORMAT_MULAW) ||
-						(SubType.Data1 == WAVE_FORMAT_IMA_ADPCM) ||
-						(SubType.Data1 == MFAudioFormat_AAC.Data1) ||
-						(SubType.Data1 == MFAudioFormat_MP3.Data1) ||
-						(SubType.Data1 == MFAudioFormat_PCM.Data1);
-
-					const bool UndocumentedFormat =
-						(SubType.Data1 == WAVE_FORMAT_WMAUDIO2) ||
-						(SubType.Data1 == WAVE_FORMAT_WMAUDIO3) ||
-						(SubType.Data1 == WAVE_FORMAT_WMAUDIO_LOSSLESS);
-
-					if (!DocumentedFormat && !UndocumentedFormat)
-					{
-						UE_LOG(LogWmfMedia, Warning, TEXT("Unsupported MP4 audio type '%s' (%s) for stream %i"), *WmfMedia::FourccToString(SubType.Data1), *WmfMedia::GuidToString(SubType), StreamIndex);
-						OutInfo += TEXT("    unsupported sub-type\n");
+					UE_LOG(LogWmfMedia, Error, TEXT("Failed to get video output sub-type for stream %i (%s)"), StreamIndex, *WmfMedia::ResultToString(Result));
+					OutInfo += FString::Printf(TEXT("        failed to get sub-type"));
 
-						return;
-					}
+					continue;
 				}
 			}
-			else if (FMemory::Memcmp(&SubType.Data2, &MFAudioFormat_Base.Data2, 12) != 0)
+
+			const uint32 BitRate = ::MFGetAttributeUINT32(MediaType, MF_MT_AVG_BITRATE, 0);
+
+			FIntPoint OutputDim;
 			{
-				if (GetDefault<UWmfMediaSettings>()->AllowNonStandardCodecs)
+				if (SUCCEEDED(::MFGetAttributeSize(MediaType, MF_MT_FRAME_SIZE, (UINT32*)&OutputDim.X, (UINT32*)&OutputDim.Y)))
 				{
-					UE_LOG(LogWmfMedia, Verbose, TEXT("Non-standard audio type '%s' (%s) for stream %i"), *WmfMedia::FourccToString(SubType.Data1), *WmfMedia::GuidToString(SubType), StreamIndex);
-					OutInfo += TEXT("    non-standard sub-type\n");
+					OutInfo += FString::Printf(TEXT("        Dimensions: %i x %i\n"), OutputDim.X, OutputDim.Y);
 				}
 				else
 				{
-					UE_LOG(LogWmfMedia, Warning, TEXT("Unsupported audio type '%s' (%s) for stream %i"), *WmfMedia::FourccToString(SubType.Data1), *WmfMedia::GuidToString(SubType), StreamIndex);
-					OutInfo += TEXT("    unsupported sub-type\n");
-
-					return;
+					OutputDim = FIntPoint::ZeroValue;
+					OutInfo += FString::Printf(TEXT("        Dimensions: n/a\n"));
 				}
 			}
 
-			// configure audio output (re-sampling fails for many media types, so we don't attempt it)
-			if (FAILED(OutputType->SetGUID(MF_MT_MAJOR_TYPE, MFMediaType_Audio)) ||
-				FAILED(OutputType->SetGUID(MF_MT_SUBTYPE, MFAudioFormat_PCM)) ||
-				FAILED(OutputType->SetUINT32(MF_MT_AUDIO_BITS_PER_SAMPLE, 16u)))
+			float FrameRate;
 			{
-				UE_LOG(LogWmfMedia, Warning, TEXT("Failed to initialize audio output type for stream %i"), StreamIndex);
-				OutInfo += TEXT("    failed to initialize output type\n");
-
-				return;
-			}
+				UINT32 Numerator = 0;
+				UINT32 Denominator = 1;
 
-			if (::MFGetAttributeUINT32(MediaType, MF_MT_AUDIO_SAMPLES_PER_SECOND, 0) != 44100)
-			{
-				UE_LOG(LogWmfMedia, Warning, TEXT("Possible loss of audio quality in stream %i due to sample rate != 44100 Hz"), StreamIndex);
+				if (SUCCEEDED(::MFGetAttributeRatio(MediaType, MF_MT_FRAME_RATE, &Numerator, &Denominator)))
+				{
+					FrameRate = static_cast<float>(Numerator) / Denominator;
+					OutInfo += FString::Printf(TEXT("        Frame Rate: %g fps\n"), FrameRate);
+				}
+				else
+				{
+					FrameRate = 0.0f;
+					OutInfo += FString::Printf(TEXT("        Frame Rate: n/a\n"));
+				}
 			}
-		}
-		else if (MajorType == MFMediaType_Binary)
-		{
-			Result = OutputType->SetGUID(MF_MT_MAJOR_TYPE, MFMediaType_Binary);
 
-			if (FAILED(Result))
+			TRange<float> FrameRates;
 			{
-				UE_LOG(LogWmfMedia, Warning, TEXT("Failed to initialize binary output type for stream %i (%s)"), StreamIndex, *WmfMedia::ResultToString(Result));
-				OutInfo += TEXT("    failed to set output type\n");
+				UINT32 Numerator = 0;
+				UINT32 Denominator = 1;
+				float Min = -1.0f;
+				float Max = -1.0f;
 
-				return;
-			}
-		}
-		else if (MajorType == MFMediaType_SAMI)
-		{
-			// configure caption output
-			Result = OutputType->SetGUID(MF_MT_MAJOR_TYPE, MFMediaType_SAMI);
+				if (SUCCEEDED(::MFGetAttributeRatio(MediaType, MF_MT_FRAME_RATE_RANGE_MIN, &Numerator, &Denominator)))
+				{
+					Min = static_cast<float>(Numerator) / Denominator;
+				}
 
-			if (FAILED(Result))
-			{
-				UE_LOG(LogWmfMedia, Warning, TEXT("Failed to initialize caption output type for stream %i (%s)"), StreamIndex, *WmfMedia::ResultToString(Result));
-				OutInfo += TEXT("    failed to set output type\n");
+				if (SUCCEEDED(::MFGetAttributeRatio(MediaType, MF_MT_FRAME_RATE_RANGE_MAX, &Numerator, &Denominator)))
+				{
+					Max = static_cast<float>(Numerator) / Denominator;
+				}
 
-				return;
-			}
-		}
-		else if (MajorType == MFMediaType_Video)
-		{
-			// filter unsupported video types
-			if (FMemory::Memcmp(&SubType.Data2, &MFVideoFormat_Base.Data2, 12) != 0)
-			{
-				if (GetDefault<UWmfMediaSettings>()->AllowNonStandardCodecs)
+				if ((Min >= 0.0f) && (Max >= 0.0f))
 				{
-					UE_LOG(LogWmfMedia, Verbose, TEXT("Non-standard video type '%s' (%s) for stream %i"), *WmfMedia::FourccToString(SubType.Data1), *WmfMedia::GuidToString(SubType), StreamIndex);
-					OutInfo += TEXT("    non-standard sub-type\n");
+					FrameRates = TRange<float>::Inclusive(Min, Max);
 				}
 				else
 				{
-					UE_LOG(LogWmfMedia, Warning, TEXT("Unsupported video type '%s' (%s) for stream %i"), *WmfMedia::FourccToString(SubType.Data1), *WmfMedia::GuidToString(SubType), StreamIndex);
-					OutInfo += TEXT("    unsupported sub-type\n");
-					return;
+					FrameRates = TRange<float>(FrameRate);
 				}
+
+				OutInfo += FString::Printf(TEXT("        Frame Rate Range: %g - %g fps\n"), FrameRates.GetLowerBoundValue(), FrameRates.GetUpperBoundValue());
 			}
 
-			if ((SubType == MFVideoFormat_HEVC) && !FWindowsPlatformMisc::VerifyWindowsVersion(10, 0))
+			FIntPoint BufferDim;
+			uint32 BufferStride;
+			EMediaTextureSampleFormat SampleFormat;
 			{
-				if ((SubType == MFVideoFormat_HEVC) && !FWindowsPlatformMisc::VerifyWindowsVersion(6, 2))
+				if (OutputSubType == MFVideoFormat_NV12)
 				{
-					UE_LOG(LogWmfMedia, Warning, TEXT("HEVC video type in stream %i requires Windows 10 or newer"), StreamIndex);
-					return;
+					BufferDim = FIntPoint(Align(OutputDim.X, 16), Align(OutputDim.Y, 16) * 3 / 2);
+					BufferStride = BufferDim.X;
+					SampleFormat = EMediaTextureSampleFormat::CharNV12;
 				}
+				else
+				{
+					long SampleStride = ::MFGetAttributeUINT32(OutputType, MF_MT_DEFAULT_STRIDE, 0);
 
-				UE_LOG(LogWmfMedia, Warning, TEXT("HEVC video type in stream %i requires Windows 10 or newer (game must be manifested for Windows 10)"), StreamIndex);
-			}
+					if (OutputSubType == MFVideoFormat_RGB32)
+					{
+						SampleFormat = EMediaTextureSampleFormat::CharBMP;
+
+						if (SampleStride == 0)
+						{
+							::MFGetStrideForBitmapInfoHeader(OutputSubType.Data1, OutputDim.X, &SampleStride);
+						}
 
-			// configure video output
-			Result = OutputType->SetGUID(MF_MT_MAJOR_TYPE, MFMediaType_Video);
+						if (SampleStride == 0)
+						{
+							SampleStride = OutputDim.X * 4;
+						}
+					}
+					else
+					{
+						SampleFormat = EMediaTextureSampleFormat::CharYUY2;
 
-			if (FAILED(Result))
-			{
-				UE_LOG(LogWmfMedia, Warning, TEXT("Failed to set video output type for stream %i (%s)"), StreamIndex, *WmfMedia::ResultToString(Result));
-				OutInfo += TEXT("    failed to set output type\n");
+						if (SampleStride == 0)
+						{
+							int32 AlignedOutputX = OutputDim.X;
 
-				return;
-			}
+							if ((SubType == MFVideoFormat_H264) || (SubType == MFVideoFormat_H264_ES))
+							{
+								AlignedOutputX = Align(AlignedOutputX, 16);
+							}
 
-			if (SubType == MFVideoFormat_HEVC)
-			{
-				Result = OutputType->SetGUID(MF_MT_SUBTYPE, MFVideoFormat_NV12);
+							SampleStride = AlignedOutputX * 2;
+						}
+					}
+
+					if (SampleStride < 0)
+					{
+						SampleStride = -SampleStride;
+					}
+
+					BufferDim = FIntPoint(SampleStride / 4, OutputDim.Y);
+					BufferStride = SampleStride;
+				}
 			}
-			else
+
+			GUID FormatType = GUID_NULL;
+
+			// prevent duplicates for legacy DirectShow media types
+			// see: https://msdn.microsoft.com/en-us/library/windows/desktop/ff485858(v=vs.85).aspx
+
+			if (SUCCEEDED(MediaType->GetGUID(MF_MT_AM_FORMAT_TYPE, &FormatType)))
 			{
-				const bool Uncompressed =
-					(SubType == MFVideoFormat_RGB555) ||
-					(SubType == MFVideoFormat_RGB565) ||
-					(SubType == MFVideoFormat_RGB24) ||
-					(SubType == MFVideoFormat_RGB32) ||
-					(SubType == MFVideoFormat_ARGB32);
-
-				Result = OutputType->SetGUID(MF_MT_SUBTYPE, Uncompressed ? MFVideoFormat_RGB32 : MFVideoFormat_YUY2);
+				if (FormatType == FORMAT_VideoInfo)
+				{
+					for (int32 Index = Track->Formats.Num() - 1; Index >= 0; --Index)
+					{
+						const FFormat& Format = Track->Formats[Index];
+
+						if ((Format.Video.FormatType == FORMAT_VideoInfo2) &&
+							(Format.Video.FrameRates == FrameRates) &&
+							(Format.Video.OutputDim == OutputDim) &&
+							(Format.TypeName == TypeName))
+						{
+							FormatIndex = Index; // keep newer format
+
+							break;
+						}
+					}
+				}
+				else if (FormatType == FORMAT_VideoInfo2)
+				{
+					for (int32 Index = Track->Formats.Num() - 1; Index >= 0; --Index)
+					{
+						FFormat& Format = Track->Formats[Index];
+
+						if ((Format.Video.FormatType == FORMAT_VideoInfo) &&
+							(Format.Video.FrameRates == FrameRates) &&
+							(Format.Video.OutputDim == OutputDim) &&
+							(Format.TypeName == TypeName))
+						{
+							Format.InputType = MediaType; // replace legacy format
+							FormatIndex = Index;
+
+							break;
+						}
+					}
+				}
 			}
 
-			if (FAILED(Result))
+			if (FormatIndex == INDEX_NONE)
 			{
-				UE_LOG(LogWmfMedia, Warning, TEXT("Failed to set video output sub-type for stream %i (%s)"), StreamIndex, *WmfMedia::ResultToString(Result));
-				OutInfo += TEXT("    failed to set output sub-type\n");
-
-				return;
+				FormatIndex = Track->Formats.Add({
+					MediaType,
+					OutputType,
+					TypeName,
+					{ 0 },
+					{
+						BitRate,
+						BufferDim,
+						BufferStride,
+						FormatType,
+						FrameRate,
+						FrameRates,
+						OutputDim,
+						SampleFormat
+					}
+				});
 			}
 		}
 		else
 		{
-			check(false); // should never get here as unsupported major types are filtered above
+			check(false); // should never get here
+		}
+
+		if (MediaType == CurrentMediaType)
+		{
+			Track->SelectedFormat = FormatIndex;
 		}
 	}
 
-	// create output sampler
-	TComPtr<FWmfMediaSampler> Sampler;
+	// get stream info
+	PWSTR OutString = NULL;
+	UINT32 OutLength = 0;
 
-	if ((MajorType != MFMediaType_Audio) || !GetDefault<UWmfMediaSettings>()->NativeAudioOut)
+	if (SUCCEEDED(StreamDescriptor->GetAllocatedString(MF_SD_LANGUAGE, &OutString, &OutLength)))
 	{
-		Sampler = new FWmfMediaSampler();
+		Track->Language = OutString;
+		::CoTaskMemFree(OutString);
 	}
 
-	// create & add track
-	FStreamInfo* StreamInfo = nullptr;
-
-	if (MajorType == MFMediaType_Audio)
+	if (SUCCEEDED(StreamDescriptor->GetAllocatedString(MF_SD_STREAM_NAME, &OutString, &OutLength)))
 	{
-		const int32 AudioTrackIndex = AudioTracks.AddDefaulted();
-		FAudioTrack& AudioTrack = AudioTracks[AudioTrackIndex];
-		{
-			AudioTrack.BitsPerSample = ::MFGetAttributeUINT32(MediaType, MF_MT_AUDIO_BITS_PER_SAMPLE, 16);
-			AudioTrack.NumChannels = ::MFGetAttributeUINT32(MediaType, MF_MT_AUDIO_NUM_CHANNELS, 0);
-			AudioTrack.SampleRate = ::MFGetAttributeUINT32(MediaType, MF_MT_AUDIO_SAMPLES_PER_SECOND, 0);
-		}
+		Track->Name = OutString;
+		::CoTaskMemFree(OutString);
+	}
 
-		if (Sampler != NULL)
-		{
-			Sampler->OnClock().AddRaw(this, &FWmfMediaTracks::HandleMediaSamplerClock, EMediaTrackType::Audio);
-			Sampler->OnSample().AddRaw(this, &FWmfMediaTracks::HandleMediaSamplerSample, EMediaTrackType::Audio);
-		}
+	Track->DisplayName = (Track->Name.IsEmpty())
+		? FText::Format(LOCTEXT("UnnamedStreamFormat", "Unnamed Track (Stream {0})"), FText::AsNumber((uint32)StreamIndex))
+		: FText::FromString(Track->Name);
 
-		OutInfo += FString::Printf(TEXT("    Channels: %i\n"), AudioTrack.NumChannels);
-		OutInfo += FString::Printf(TEXT("    Sample Rate: %i Hz\n"), AudioTrack.SampleRate);
-		OutInfo += FString::Printf(TEXT("    Bits Per Sample: %i\n"), AudioTrack.BitsPerSample);
+	Track->Descriptor = StreamDescriptor;
+	Track->Handler = Handler;
+	Track->Protected = Protected;
+	Track->StreamIndex = StreamIndex;
 
-		StreamInfo = &AudioTrack;
-		CurrentTrackIndex = AudioTrackIndex;
+	if (Track->Formats.Num() > 0)
+	{
+		Track->SelectedFormat = 0;
 	}
-	else if (MajorType == MFMediaType_SAMI)
+	else
 	{
-		const int32 CaptionTrackIndex = CaptionTracks.AddDefaulted();
-		FCaptionTrack& CaptionTrack = CaptionTracks[CaptionTrackIndex];
+		Track->SelectedFormat = INDEX_NONE;
 
-		if (Sampler != NULL)
-		{
-			Sampler->OnSample().AddRaw(this, &FWmfMediaTracks::HandleMediaSamplerSample, EMediaTrackType::Caption);
-		}
-
-		StreamInfo = &CaptionTrack;
-		CurrentTrackIndex = CaptionTrackIndex;
+		UE_LOG(LogWmfMedia, Warning, TEXT("No supported media type in stream %i of type %s"), StreamIndex, *WmfMedia::MajorTypeToString(MajorType));
+		OutInfo += TEXT("    unsupported media type\n");
 	}
-	else if (MajorType == MFMediaType_Binary)
+}
+
+
+const FWmfMediaTracks::FFormat* FWmfMediaTracks::GetAudioFormat(int32 TrackIndex, int32 FormatIndex) const
+{
+	if (AudioTracks.IsValidIndex(TrackIndex))
 	{
-		const int32 BinaryTrackIndex = BinaryTracks.AddDefaulted();
-		FBinaryTrack& BinaryTrack = BinaryTracks[BinaryTrackIndex];
+		const FTrack& Track = AudioTracks[TrackIndex];
 
-		if (Sampler != NULL)
+		if (Track.Formats.IsValidIndex(FormatIndex))
 		{
-			Sampler->OnSample().AddRaw(this, &FWmfMediaTracks::HandleMediaSamplerSample, EMediaTrackType::Metadata);
+			return &Track.Formats[FormatIndex];
 		}
-
-		StreamInfo = &BinaryTrack;
-		CurrentTrackIndex = BinaryTrackIndex;
 	}
-	else if ((MajorType == MFMediaType_Video) || (MajorType == MFMediaType_Image))
+
+	return nullptr;
+}
+
+
+const FWmfMediaTracks::FTrack* FWmfMediaTracks::GetTrack(EMediaTrackType TrackType, int32 TrackIndex) const
+{
+	switch (TrackType)
 	{
-		GUID OutputSubType;
+	case EMediaTrackType::Audio:
+		if (AudioTracks.IsValidIndex(TrackIndex))
 		{
-			HRESULT Result = OutputType->GetGUID(MF_MT_SUBTYPE, &OutputSubType);
-
-			if (FAILED(Result))
-			{
-				UE_LOG(LogWmfMedia, Error, TEXT("Failed to get video output sub-type for stream %i (%s)"), StreamIndex, *WmfMedia::ResultToString(Result));
-				return;
-			}
+			return &AudioTracks[TrackIndex];
 		}
 
-		const int32 VideoTrackIndex = VideoTracks.AddDefaulted();
-		FVideoTrack& VideoTrack = VideoTracks[VideoTrackIndex];
+	case EMediaTrackType::Metadata:
+		if (MetadataTracks.IsValidIndex(TrackIndex))
 		{
-			// bit rate
-			VideoTrack.BitRate = ::MFGetAttributeUINT32(MediaType, MF_MT_AVG_BITRATE, 0);
-			
-			// dimensions
-			if (SUCCEEDED(::MFGetAttributeSize(MediaType, MF_MT_FRAME_SIZE, (UINT32*)&VideoTrack.OutputDim.X, (UINT32*)&VideoTrack.OutputDim.Y)))
-			{
-				OutInfo += FString::Printf(TEXT("    Dimensions: %i x %i\n"), VideoTrack.OutputDim.X, VideoTrack.OutputDim.Y);
-			}
-			else
-			{
-				VideoTrack.OutputDim = FIntPoint::ZeroValue;
-				OutInfo += FString::Printf(TEXT("    Dimensions: n/a\n"));
-			}
-
-			// frame rate
-			UINT32 Numerator = 0;
-			UINT32 Denominator = 1;
-
-			if (SUCCEEDED(::MFGetAttributeRatio(MediaType, MF_MT_FRAME_RATE, &Numerator, &Denominator)))
-			{
-				VideoTrack.FrameRate = static_cast<float>(Numerator) / Denominator;
-				OutInfo += FString::Printf(TEXT("    Frame Rate: %g fps\n"), VideoTrack.FrameRate);
-			}
-			else
-			{
-				VideoTrack.FrameRate = 0.0f;
-				OutInfo += FString::Printf(TEXT("    Frame Rate: n/a\n"));
-			}
-
-			// sink format
-			if (OutputSubType == MFVideoFormat_NV12)
-			{
-				VideoTrack.BufferDim = FIntPoint(Align(VideoTrack.OutputDim.X, 16), VideoTrack.OutputDim.Y * 3 / 2);
-				VideoTrack.SinkFormat = EMediaTextureSinkFormat::CharNV12;
-			}
-			else
-			{
-				long SampleStride = ::MFGetAttributeUINT32(MediaType, MF_MT_DEFAULT_STRIDE, 0);
-
-				if (OutputSubType == MFVideoFormat_RGB32)
-				{
-					VideoTrack.SinkFormat = EMediaTextureSinkFormat::CharBMP;
-
-					if (SampleStride == 0)
-					{
-						::MFGetStrideForBitmapInfoHeader(SubType.Data1, VideoTrack.OutputDim.X, &SampleStride);
-					}
-
-					if (SampleStride == 0)
-					{
-						SampleStride = VideoTrack.OutputDim.X * 4;
-					}
-				}
-				else
-				{
-					VideoTrack.SinkFormat = EMediaTextureSinkFormat::CharYUY2;
-
-					if (SampleStride == 0)
-					{
-						int32 AlignedOutputX = VideoTrack.OutputDim.X;
-
-						if ((SubType == MFVideoFormat_H264) || (SubType == MFVideoFormat_H264_ES))
-						{
-							AlignedOutputX = Align(AlignedOutputX, 16);
-						}
-
-						SampleStride = AlignedOutputX * 2;
-					}
-				}
-
-				if (SampleStride < 0)
-				{
-					SampleStride = -SampleStride;
-				}
+			return &MetadataTracks[TrackIndex];
+		}
 
-				VideoTrack.BufferDim = FIntPoint(SampleStride / 4, VideoTrack.OutputDim.Y);
-			}
+	case EMediaTrackType::Caption:
+		if (CaptionTracks.IsValidIndex(TrackIndex))
+		{
+			return &CaptionTracks[TrackIndex];
 		}
 
-		if (Sampler != NULL)
+	case EMediaTrackType::Video:
+		if (VideoTracks.IsValidIndex(TrackIndex))
 		{
-			Sampler->OnSample().AddRaw(this, &FWmfMediaTracks::HandleMediaSamplerSample, EMediaTrackType::Video);
+			return &VideoTracks[TrackIndex];
 		}
 
-		StreamInfo = &VideoTrack;
-		CurrentTrackIndex = VideoTrackIndex;
+	default:
+		break; // unsupported track type
 	}
 
-	// get stream info
-	if (StreamInfo != nullptr)
+	return nullptr;
+}
+
+
+const FWmfMediaTracks::FFormat* FWmfMediaTracks::GetVideoFormat(int32 TrackIndex, int32 FormatIndex) const
+{
+	if (VideoTracks.IsValidIndex(TrackIndex))
 	{
-		PWSTR OutString = NULL;
-		UINT32 OutLength;
+		const FTrack& Track = VideoTracks[TrackIndex];
 
-		if (SUCCEEDED(StreamDescriptor->GetAllocatedString(MF_SD_LANGUAGE, &OutString, &OutLength)))
+		if (Track.Formats.IsValidIndex(FormatIndex))
 		{
-			StreamInfo->Language = OutString;
-			::CoTaskMemFree(OutString);
+			return &Track.Formats[FormatIndex];
 		}
+	}
 
-		if (SUCCEEDED(StreamDescriptor->GetAllocatedString(MF_SD_STREAM_NAME, &OutString, &OutLength)))
-		{
-			StreamInfo->Name = OutString;
-			::CoTaskMemFree(OutString);
-		}
+	return nullptr;
+}
 
-		StreamInfo->DisplayName = (StreamInfo->Name.IsEmpty())
-			? FText::Format(LOCTEXT("UnnamedStreamFormat", "Unnamed Track {0} (Stream {1})"), FText::AsNumber(CurrentTrackIndex), FText::AsNumber((uint32)StreamIndex))
-			: FText::FromString(StreamInfo->Name);
 
-		StreamInfo->Descriptor = StreamDescriptor;
-		StreamInfo->OutputType = OutputType;
-		StreamInfo->Protected = Protected;
-		StreamInfo->Sampler = Sampler;
-		StreamInfo->StreamIndex = StreamIndex;
-	}
+/* FWmfMediaTracks callbacks
+ *****************************************************************************/
+
+void FWmfMediaTracks::HandleMediaSamplerClock(EWmfMediaSamplerClockEvent Event, EMediaTrackType TrackType)
+{
+	// IMFSampleGrabberSinkCallback callbacks seem to be broken (always returns Stopped)
+	// We handle sink synchronization via SetPaused() as a workaround
 }
 
 
-void FWmfMediaTracks::InitializeAudioSink()
+void FWmfMediaTracks::HandleMediaSamplerAudioSample(const uint8* Buffer, uint32 Size, FTimespan /*Duration*/, FTimespan Time)
 {
-	if ((AudioSink == nullptr) || (SelectedAudioTrack == INDEX_NONE))
+	if (Buffer == nullptr)
 	{
 		return;
 	}
 
-	const FAudioTrack& AudioTrack = AudioTracks[SelectedAudioTrack];
-	AudioSink->InitializeAudioSink(AudioTrack.NumChannels, AudioTrack.SampleRate);
-}
+	FScopeLock Lock(&CriticalSection);
 
+	if (!AudioTracks.IsValidIndex(SelectedAudioTrack))
+	{
+		return; // invalid track index
+	}
 
-void FWmfMediaTracks::InitializeMetadataSink()
-{
-	if ((MetadataSink == nullptr) || (SelectedBinaryTrack == INDEX_NONE))
+	const FTrack& Track = AudioTracks[SelectedAudioTrack];
+	const FFormat* Format = GetAudioFormat(SelectedAudioTrack, Track.SelectedFormat);
+
+	if (Format == nullptr)
 	{
-		return;
+		return; // no format selected
 	}
 
-	const FBinaryTrack& BinaryTrack = BinaryTracks[SelectedBinaryTrack];
-	MetadataSink->InitializeBinarySink();
+	// duration provided by WMF is sometimes incorrect when seeking
+	FTimespan Duration = (Size * ETimespan::TicksPerSecond) / (Format->Audio.NumChannels * Format->Audio.SampleRate * sizeof(int16));
+
+	// create & add sample to queue
+	const TSharedRef<FWmfMediaAudioSample, ESPMode::ThreadSafe> AudioSample = AudioSamplePool->AcquireShared();
+
+	if (AudioSample->Initialize(Buffer, Size, Format->Audio.NumChannels, Format->Audio.SampleRate, Time, Duration))
+	{
+		AudioSampleQueue.Enqueue(AudioSample);
+	}
 }
 
 
-void FWmfMediaTracks::InitializeOverlaySink()
+void FWmfMediaTracks::HandleMediaSamplerCaptionSample(const uint8* Buffer, uint32 Size, FTimespan Duration, FTimespan Time)
 {
-	if ((OverlaySink == nullptr) || (SelectedCaptionTrack == INDEX_NONE))
+	if (Buffer == nullptr)
 	{
 		return;
 	}
 
-	OverlaySink->InitializeOverlaySink();
+	FScopeLock Lock(&CriticalSection);
+
+	if (!CaptionTracks.IsValidIndex(SelectedCaptionTrack))
+	{
+		return; // invalid track index
+	}
+
+	// create & add sample to queue
+	const FTrack& Track = CaptionTracks[SelectedCaptionTrack];
+	const auto CaptionSample = MakeShared<FWmfMediaOverlaySample, ESPMode::ThreadSafe>();
+
+	if (CaptionSample->Initialize((char*)Buffer, Time, Duration))
+	{
+		CaptionSampleQueue.Enqueue(CaptionSample);
+	}
 }
 
 
-void FWmfMediaTracks::InitializeVideoSink()
+void FWmfMediaTracks::HandleMediaSamplerMetadataSample(const uint8* Buffer, uint32 Size, FTimespan Duration, FTimespan Time)
 {
-	if ((VideoSink == nullptr) || (SelectedVideoTrack == INDEX_NONE))
+	if (Buffer == nullptr)
 	{
 		return;
 	}
 
-	const FVideoTrack& VideoTrack = VideoTracks[SelectedVideoTrack];
-	VideoSink->InitializeTextureSink(VideoTrack.OutputDim, VideoTrack.BufferDim, VideoTrack.SinkFormat, EMediaTextureSinkMode::Buffered);
-}
+	FScopeLock Lock(&CriticalSection);
 
+	if (!MetadataTracks.IsValidIndex(SelectedBinaryTrack))
+	{
+		return; // invalid track index
+	}
 
-/* FWmfMediaTracks event handlers
- *****************************************************************************/
+	// create & add sample to queue
+	const FTrack& Track = MetadataTracks[SelectedBinaryTrack];
+	const auto BinarySample = MakeShared<FWmfMediaBinarySample, ESPMode::ThreadSafe>();
 
-void FWmfMediaTracks::HandleMediaSamplerClock(EWmfMediaSamplerClockEvent Event, EMediaTrackType TrackType)
-{
-	// IMFSampleGrabberSinkCallback callbacks seem to be broken (always returns Stopped)
-	// We handle sink synchronization via SetPaused() as a workaround
+	if (BinarySample->Initialize(Buffer, Size, Time, Duration))
+	{
+		MetadataSampleQueue.Enqueue(BinarySample);
+	}
 }
 
 
-void FWmfMediaTracks::HandleMediaSamplerSample(const uint8* Buffer, uint32 Size, FTimespan Duration, FTimespan Time, EMediaTrackType TrackType)
+void FWmfMediaTracks::HandleMediaSamplerVideoSample(const uint8* Buffer, uint32 Size, FTimespan Duration, FTimespan Time)
 {
-	FScopeLock Lock(&CriticalSection);
-
 	if (Buffer == nullptr)
 	{
 		return;
 	}
 
-	// forward the sample to the sink
-	switch (TrackType)
+	FScopeLock Lock(&CriticalSection);
+
+	if (!VideoTracks.IsValidIndex(SelectedVideoTrack))
 	{
-	case EMediaTrackType::Audio:
-		if ((AudioSink != nullptr) && AudioTracks.IsValidIndex(SelectedAudioTrack))
-		{
-			AudioSink->PlayAudioSink(Buffer, Size, Time);
-		}
-		break;
+		return; // invalid track index
+	}
 
-	case EMediaTrackType::Metadata:
-		if ((MetadataSink != nullptr) && BinaryTracks.IsValidIndex(SelectedBinaryTrack))
-		{
-			MetadataSink->ProcessBinarySinkData(Buffer, Size, Time, Duration);
-		}
-		break;
+	// create & add sample to queue
+	const FTrack& Track = VideoTracks[SelectedVideoTrack];
+	const FFormat* Format = GetVideoFormat(SelectedVideoTrack, Track.SelectedFormat);
 
-	case EMediaTrackType::Caption:
-		if ((OverlaySink != nullptr) && CaptionTracks.IsValidIndex(SelectedCaptionTrack))
-		{
-			OverlaySink->AddOverlaySinkText(FText::FromString((TCHAR*)Buffer), EMediaOverlayType::Caption, Time, Duration, TOptional<FVector2D>());
-		}
-		break;
+	if (Format == nullptr)
+	{
+		return; // no format selected
+	}
 
-	case EMediaTrackType::Video:
-		if ((VideoSink != nullptr) && VideoTracks.IsValidIndex(SelectedVideoTrack))
-		{
-			VideoSink->UpdateTextureSinkBuffer(Buffer);
-			VideoSink->DisplayTextureSinkBuffer(Time);
-		}
-		break;
+	if ((Format->Video.BufferStride * Format->Video.BufferDim.Y) > Size)
+	{
+		return; // invalid buffer size (can happen during format switch)
+	}
+
+	const TSharedRef<FWmfMediaTextureSample, ESPMode::ThreadSafe> TextureSample = VideoSamplePool->AcquireShared();
+
+	if (TextureSample->Initialize(
+		Buffer,
+		Size,
+		Format->Video.BufferDim,
+		Format->Video.OutputDim,
+		Format->Video.SampleFormat,
+		Format->Video.BufferStride,
+		Time,
+		Duration))
+	{
+		VideoSampleQueue.Enqueue(TextureSample);
 	}
 }
 
diff --git a/Engine/Plugins/Media/WmfMedia/Source/WmfMedia/Private/Player/WmfMediaTracks.h b/Engine/Plugins/Media/WmfMedia/Source/WmfMedia/Private/Player/WmfMediaTracks.h
--- a/Engine/Plugins/Media/WmfMedia/Source/WmfMedia/Private/Player/WmfMediaTracks.h
+++ b/Engine/Plugins/Media/WmfMedia/Source/WmfMedia/Private/Player/WmfMediaTracks.h
@@ -2,77 +2,93 @@
 
 #pragma once
 
-#include "../WmfMediaPrivate.h"
-#include "IMediaOutput.h"
-#include "IMediaTextureSink.h"
-#include "IMediaTracks.h"
+#include "WmfMediaPrivate.h"
 
 #if WMFMEDIA_SUPPORTED_PLATFORM
 
-#include "AllowWindowsPlatformTypes.h"
+#include "CoreTypes.h"
+#include "Containers/Array.h"
+#include "Containers/UnrealString.h"
+#include "Internationalization/Text.h"
+#include "IMediaSamples.h"
+#include "IMediaTracks.h"
+#include "Math/IntPoint.h"
+#include "MediaSampleQueue.h"
+#include "Templates/SharedPointer.h"
 
+#include "AllowWindowsPlatformTypes.h"
 
+class FWmfMediaAudioSamplePool;
 class FWmfMediaSampler;
-class IMediaAudioSink;
-class IMediaBinarySink;
-class IMediaOptions;
-class IMediaStringSink;
-class IMediaTextureSink;
+class FWmfMediaTextureSamplePool;
+class IMediaAudioSample;
+class IMediaBinarySample;
+class IMediaOverlaySample;
+class IMediaTextureSample;
 
+enum class EMediaTextureSampleFormat;
 enum class EMediaTrackType;
 enum class EWmfMediaSamplerClockEvent;
 
 
 /**
- * Implements the track collection for Windows Media Framework based media players.
+ * Track collection for Windows Media Foundation based media players.
  */
 class FWmfMediaTracks
-	: public IMediaOutput
+	: public IMediaSamples
 	, public IMediaTracks
 {
-	struct FStreamInfo
+	/** Track format. */
+	struct FFormat
+	{
+		TComPtr<IMFMediaType> InputType;
+		TComPtr<IMFMediaType> OutputType;
+		FString TypeName;
+
+		struct
+		{
+			uint32 BitsPerSample;
+			uint32 NumChannels;
+			uint32 SampleRate;
+		}
+		Audio;
+
+		struct
+		{
+			uint32 BitRate;
+			FIntPoint BufferDim;
+			uint32 BufferStride;
+			GUID FormatType;
+			float FrameRate;
+			TRange<float> FrameRates;
+			FIntPoint OutputDim;
+			EMediaTextureSampleFormat SampleFormat;
+		}
+		Video;
+	};
+
+
+	/** Track information. */
+	struct FTrack
 	{
 		TComPtr<IMFStreamDescriptor> Descriptor;
 		FText DisplayName;
+		TArray<FFormat> Formats;
+		TComPtr<IMFMediaTypeHandler> Handler;
 		FString Language;
 		FString Name;
-		TComPtr<IMFMediaType> OutputType;
 		bool Protected;
-		TComPtr<FWmfMediaSampler> Sampler;
+		int32 SelectedFormat;
 		DWORD StreamIndex;
 	};
-	
-	struct FAudioTrack : public FStreamInfo
-	{
-		uint32 BitsPerSample;
-		uint32 NumChannels;
-		uint32 SampleRate;
-	};
 
-	struct FBinaryTrack : public FStreamInfo
-	{
-	};
-	
-	struct FCaptionTrack : public FStreamInfo
-	{
-	};
-		
-	struct FVideoTrack : public FStreamInfo
-	{
-		uint32 BitRate;
-		FIntPoint BufferDim;
-		float FrameRate;
-		FIntPoint OutputDim;
-		EMediaTextureSinkFormat SinkFormat;
-	};
-	
 public:
 
 	/** Default constructor. */
 	FWmfMediaTracks();
 
 	/** Virtual destructor. */
-	~FWmfMediaTracks();
+	virtual ~FWmfMediaTracks();
 
 public:
 
@@ -83,102 +99,150 @@ public:
 	 */
 	void AppendStats(FString &OutStats) const;
 
+	/**
+	 * Clear the streams flags.
+	 *
+	 * @see GetFlags
+	 */
+	void ClearFlags();
+
 	/**
 	 * Create the playback topology for the current track selection.
 	 *
 	 * @return The new topology.
 	 */
-	TComPtr<IMFTopology> CreateTopology() const;
-
-	/** Discard any pending data on the media sinks. */
-	void FlushSinks();
+	TComPtr<IMFTopology> CreateTopology();
 
 	/**
-	 * Initialize the track collection.
+	 * Get the total duration of the current media source.
 	 *
-	 * @param InMediaSource The media source object.
-	 * @param InPresentationDescriptor The presentation descriptor object.
-	 * @param OutInfo Will contain information about the available media tracks.
-	 * @return true on success, false otherwise.
+	 * @return Media duration.
+	 * @see GetInfo, GetSamples, Initialize
 	 */
-	bool Initialize(IMFMediaSource& InMediaSource, IMFPresentationDescriptor& InPresentationDescriptor, FString& OutInfo);
+	FTimespan GetDuration() const;
 
-	/** Remove all streams from the track collection. */
-	void Reset();
+	/**
+	 * Get the current flags.
+	 *
+	 * @param OutMediaSourceChanged Will indicate whether the media source changed.
+	 * @param OutTopologyChanged Will indicate whether the topology changed.
+	 * @see ClearFlags
+	 */
+	void GetFlags(bool& OutMediaSourceChanged, bool& OutTopologyChanged) const;
 
 	/**
-	 * Notify tracks that playback was paused or resumed.
+	 * Get the information string for the currently loaded media source.
 	 *
-	 * @param Paused true if playback was paused, false if resumed.
+	 * @return Info string.
+	 * @see GetDuration, GetSamples
 	 */
-	void SetPaused(bool Paused);
+	const FString& GetInfo() const
+	{
+		return Info;
+	}
 
-public:
+	/**
+	 * Initialize the track collection.
+	 *
+	 * @param InMediaSource The media source object.
+	 * @see IsInitialized, Shutdown
+	 */
+	void Initialize(IMFMediaSource* InMediaSource);
 
-	/** Get an event that gets fired when the track selection changed. */
-	DECLARE_EVENT(FWmfMediaTracks, FOnSelectionChanged)
-	FOnSelectionChanged& OnSelectionChanged()
+	/**
+	 * Whether this object has been initialized.
+	 *
+	 * @return true if initialized, false otherwise.
+	 * @see Initialize, Shutdown
+	 */
+	bool IsInitialized() const
 	{
-		return SelectionChangedEvent;
+		return (MediaSource != NULL);
 	}
 
+	/**
+	 * Shut down the track collection.
+	 *
+	 * @see Initialize, IsInitialized
+	 */
+	void Shutdown();
+
 public:
 
-	//~ IMediaOutput interface
+	//~ IMediaSamples interface
 
-	virtual void SetAudioSink(IMediaAudioSink* Sink) override;
-	virtual void SetMetadataSink(IMediaBinarySink* Sink) override;
-	virtual void SetOverlaySink(IMediaOverlaySink* Sink) override;
-	virtual void SetVideoSink(IMediaTextureSink* Sink) override;
+	virtual bool FetchAudio(TRange<FTimespan> TimeRange, TSharedPtr<IMediaAudioSample, ESPMode::ThreadSafe>& OutSample) override;
+	virtual bool FetchCaption(TRange<FTimespan> TimeRange, TSharedPtr<IMediaOverlaySample, ESPMode::ThreadSafe>& OutSample) override;
+	virtual bool FetchMetadata(TRange<FTimespan> TimeRange, TSharedPtr<IMediaBinarySample, ESPMode::ThreadSafe>& OutSample) override;
+	virtual bool FetchVideo(TRange<FTimespan> TimeRange, TSharedPtr<IMediaTextureSample, ESPMode::ThreadSafe>& OutSample) override;
+	virtual void FlushSamples() override;
 
 public:
 
 	//~ IMediaTracks interface
 
-	virtual uint32 GetAudioTrackChannels(int32 TrackIndex) const override;
-	virtual uint32 GetAudioTrackSampleRate(int32 TrackIndex) const override;
+	virtual bool GetAudioTrackFormat(int32 TrackIndex, int32 FormatIndex, FMediaAudioTrackFormat& OutFormat) const override;
 	virtual int32 GetNumTracks(EMediaTrackType TrackType) const override;
+	virtual int32 GetNumTrackFormats(EMediaTrackType TrackType, int32 TrackIndex) const override;
 	virtual int32 GetSelectedTrack(EMediaTrackType TrackType) const override;
 	virtual FText GetTrackDisplayName(EMediaTrackType TrackType, int32 TrackIndex) const override;
+	virtual int32 GetTrackFormat(EMediaTrackType TrackType, int32 TrackIndex) const override;
 	virtual FString GetTrackLanguage(EMediaTrackType TrackType, int32 TrackIndex) const override;
 	virtual FString GetTrackName(EMediaTrackType TrackType, int32 TrackIndex) const override;
-	virtual uint32 GetVideoTrackBitRate(int32 TrackIndex) const override;
-	virtual FIntPoint GetVideoTrackDimensions(int32 TrackIndex) const override;
-	virtual float GetVideoTrackFrameRate(int32 TrackIndex) const override;
+	virtual bool GetVideoTrackFormat(int32 TrackIndex, int32 FormatIndex, FMediaVideoTrackFormat& OutFormat) const override;
 	virtual bool SelectTrack(EMediaTrackType TrackType, int32 TrackIndex) override;
+	virtual bool SetTrackFormat(EMediaTrackType TrackType, int32 TrackIndex, int32 FormatIndex) override;
+	virtual bool SetVideoTrackFrameRate(int32 TrackIndex, int32 FormatIndex, float FrameRate) override;
 
 protected:
 
 	/**
-	 * Add the given stream to the specified playback topology.
+	 * Add the specified stream to the track collection.
 	 *
-	 * @param Stream The stream to add.
-	 * @param Topology The playback topology.
-	 * @return true on success, false otherwise.
+	 * @param StreamIndex The index of the stream to add.
+	 * @param OutInfo Will contain appended debug information.
 	 */
-	bool AddStreamToTopology(const FStreamInfo& Stream, IMFTopology& Topology) const;
+	void AddStreamToTracks(uint32 StreamIndex, FString& OutInfo);
 
 	/**
-	 * Add the specified stream to the track collection.
+	 * Add the given track to the specified playback topology.
 	 *
-	 * @param StreamIndex The index of the stream to add.
-	 * @param InMediaSource The media source containing the stream.
-	 * @param InPresentationDescriptor The presentation descriptor containing the stream's details.
-	 * @param OutInfo Will contain appended debug information.
+	 * @param Track The track to add.
+	 * @param Topology The playback topology.
+	 * @return true on success, false otherwise.
 	 */
-	void AddStreamToTracks(uint32 StreamIndex, IMFMediaSource& InMediaSource, IMFPresentationDescriptor& InPresentationDescriptor, FString& OutInfo);
+	bool AddTrackToTopology(const FTrack& Track, IMFTopology& Topology) const;
 
-	/** Initialize the current audio sink. */
-	void InitializeAudioSink();
+private:
 
-	/** Initialize the current metadata sink. */
-	void InitializeMetadataSink();
+	/**
+	 * Get the specified audio format.
+	 *
+	 * @param TrackIndex Index of the audio track that contains the format.
+	 * @param FormatIndex Index of the format to return.
+	 * @return Pointer to format, or nullptr if not found.
+	 * @see GetVideoFormat
+	 */
+	const FFormat* GetAudioFormat(int32 TrackIndex, int32 FormatIndex) const;
 
-	/** Initialize the current text overlay sink. */
-	void InitializeOverlaySink();
+	/**
+	 * Get the specified track information.
+	 *
+	 * @param TrackType The type of track.
+	 * @param TrackIndex Index of the track to return.
+	 * @return Pointer to track, or nullptr if not found.
+	 */
+	const FTrack* GetTrack(EMediaTrackType TrackType, int32 TrackIndex) const;
 
-	/** Initialize the current video sink. */
-	void InitializeVideoSink();
+	/**
+	 * Get the specified video format.
+	 *
+	 * @param TrackIndex Index of the video track that contains the format.
+	 * @param FormatIndex Index of the format to return.
+	 * @return Pointer to format, or nullptr if not found.
+	 * @see GetAudioFormat
+	 */
+	const FFormat* GetVideoFormat(int32 TrackIndex, int32 FormatIndex) const;
 
 private:
 
@@ -186,37 +250,54 @@ private:
 	void HandleMediaSamplerClock(EWmfMediaSamplerClockEvent Event, EMediaTrackType TrackType);
 
 	/** Callback for handling new samples from the streams' media sample buffers. */
-	void HandleMediaSamplerSample(const uint8* Buffer, uint32 Size, FTimespan Duration, FTimespan Time, EMediaTrackType TrackType);
+	void HandleMediaSamplerAudioSample(const uint8* Buffer, uint32 Size, FTimespan Duration, FTimespan Time);
 
-private:
+	/** Callback for handling new caption samples. */
+	void HandleMediaSamplerCaptionSample(const uint8* Buffer, uint32 Size, FTimespan Duration, FTimespan Time);
 
-	/** The currently used audio sink. */
-	IMediaAudioSink* AudioSink;
+	/** Callback for handling new metadata samples. */
+	void HandleMediaSamplerMetadataSample(const uint8* Buffer, uint32 Size, FTimespan Duration, FTimespan Time);
 
-	/** The currently used metadata sink. */
-	IMediaBinarySink* MetadataSink;
+	/** Callback for handling new video samples. */
+	void HandleMediaSamplerVideoSample(const uint8* Buffer, uint32 Size, FTimespan Duration, FTimespan Time);
 
-	/** The currently used text overlay sink. */
-	IMediaOverlaySink* OverlaySink;
+private:
 
-	/** The currently used video sink. */
-	IMediaTextureSink* VideoSink;
+	/** Audio sample object pool. */
+	FWmfMediaAudioSamplePool* AudioSamplePool;
 
-private:
+	/** Audio sample queue. */
+	TMediaSampleQueue<IMediaAudioSample> AudioSampleQueue;
 
 	/** The available audio tracks. */
-	TArray<FAudioTrack> AudioTracks;
+	TArray<FTrack> AudioTracks;
 
-	/** The available binary tracks. */
-	TArray<FBinaryTrack> BinaryTracks;
+	/** Overlay sample queue. */
+	TMediaSampleQueue<IMediaOverlaySample> CaptionSampleQueue;
 
 	/** The available caption tracks. */
-	TArray<FCaptionTrack> CaptionTracks;
+	TArray<FTrack> CaptionTracks;
 
-	/** The available video tracks. */
-	TArray<FVideoTrack> VideoTracks;
+	/** Synchronizes write access to track arrays, selections & sinks. */
+	mutable FCriticalSection CriticalSection;
 
-private:
+	/** Media information string. */
+	FString Info;
+
+	/** The currently opened media. */
+	TComPtr<IMFMediaSource> MediaSource;
+
+	/** Whether the media source has changed. */
+	bool MediaSourceChanged;
+
+	/** Metadata sample queue. */
+	TMediaSampleQueue<IMediaBinarySample> MetadataSampleQueue;
+
+	/** The available metadata tracks. */
+	TArray<FTrack> MetadataTracks;
+
+	/** The presentation descriptor of the currently opened media. */
+	TComPtr<IMFPresentationDescriptor> PresentationDescriptor;
 
 	/** Index of the selected audio track. */
 	int32 SelectedAudioTrack;
@@ -230,22 +311,20 @@ private:
 	/** Index of the selected video track. */
 	int32 SelectedVideoTrack;
 
-private:
+	/** Whether the playback topology changed. */
+	bool TopologyChanged;
 
-	/** Synchronizes write access to track arrays, selections & sinks. */
-	mutable FCriticalSection CriticalSection;
+	/** Video sample object pool. */
+	FWmfMediaTextureSamplePool* VideoSamplePool;
 
-	/** The currently opened media. */
-	TComPtr<IMFMediaSource> MediaSource;
+	/** Video sample queue. */
+	TMediaSampleQueue<IMediaTextureSample> VideoSampleQueue;
 
-	/** The presentation descriptor of the currently opened media. */
-	TComPtr<IMFPresentationDescriptor> PresentationDescriptor;
-
-	/** Event that gets fired when the track selection changed. */
-	FOnSelectionChanged SelectionChangedEvent;
+	/** The available video tracks. */
+	TArray<FTrack> VideoTracks;
 };
 
 
 #include "HideWindowsPlatformTypes.h"
 
-#endif
+#endif //WMFMEDIA_SUPPORTED_PLATFORM
diff --git a/Engine/Plugins/Media/WmfMedia/Source/WmfMedia/Private/Wmf/WmfMediaByteStream.cpp b/Engine/Plugins/Media/WmfMedia/Source/WmfMedia/Private/Wmf/WmfMediaByteStream.cpp
--- a/Engine/Plugins/Media/WmfMedia/Source/WmfMedia/Private/Wmf/WmfMediaByteStream.cpp
+++ b/Engine/Plugins/Media/WmfMedia/Source/WmfMedia/Private/Wmf/WmfMediaByteStream.cpp
@@ -1,13 +1,15 @@
 // Copyright 1998-2017 Epic Games, Inc. All Rights Reserved.
 
 #include "WmfMediaByteStream.h"
-#include "Misc/ScopeLock.h"
-#include "WmfMediaReadState.h"
 
 #if WMFMEDIA_SUPPORTED_PLATFORM
 
+#include "WmfMediaReadState.h"
+#include "Misc/ScopeLock.h"
+
 #include "AllowWindowsPlatformTypes.h"
 
+
 /* FWmfByteStream structors
  *****************************************************************************/
 
@@ -18,6 +20,12 @@ FWmfMediaByteStream::FWmfMediaByteStream(const TSharedRef<FArchive, ESPMode::Thr
 { }
 
 
+FWmfMediaByteStream::~FWmfMediaByteStream()
+{
+	check(RefCount == 0);
+}
+
+
 /* IMFAsyncCallback interface
  *****************************************************************************/
 
diff --git a/Engine/Plugins/Media/WmfMedia/Source/WmfMedia/Private/Wmf/WmfMediaByteStream.h b/Engine/Plugins/Media/WmfMedia/Source/WmfMedia/Private/Wmf/WmfMediaByteStream.h
--- a/Engine/Plugins/Media/WmfMedia/Source/WmfMedia/Private/Wmf/WmfMediaByteStream.h
+++ b/Engine/Plugins/Media/WmfMedia/Source/WmfMedia/Private/Wmf/WmfMediaByteStream.h
@@ -2,7 +2,7 @@
 
 #pragma once
 
-#include "../WmfMediaPrivate.h"
+#include "WmfMediaPrivate.h"
 
 #if WMFMEDIA_SUPPORTED_PLATFORM
 
@@ -25,9 +25,6 @@ public:
 	 */
 	FWmfMediaByteStream(const TSharedRef<FArchive, ESPMode::ThreadSafe>& InArchive);
 
-	/** Virtual destructor. */
-	virtual ~FWmfMediaByteStream() { }
-
 public:
 
 	//~ IMFAsyncCallback interface
@@ -58,6 +55,11 @@ public:
 	STDMETHODIMP SetCurrentPosition(QWORD qwPosition);
 	STDMETHODIMP Write(const BYTE* pb, ULONG cb, ULONG* pcbWritten);
 
+private:
+
+	/** Hidden destructor. */
+	virtual ~FWmfMediaByteStream();
+
 private:
 
 	/** Whether the stream is currently being read asynchronously. */
@@ -76,4 +78,4 @@ private:
 
 #include "HideWindowsPlatformTypes.h"
 
-#endif
+#endif //WMFMEDIA_SUPPORTED_PLATFORM
diff --git a/Engine/Plugins/Media/WmfMedia/Source/WmfMedia/Private/Wmf/WmfMediaReadState.h b/Engine/Plugins/Media/WmfMedia/Source/WmfMedia/Private/Wmf/WmfMediaReadState.h
--- a/Engine/Plugins/Media/WmfMedia/Source/WmfMedia/Private/Wmf/WmfMediaReadState.h
+++ b/Engine/Plugins/Media/WmfMedia/Source/WmfMedia/Private/Wmf/WmfMediaReadState.h
@@ -30,9 +30,6 @@ public:
 		, RefCount(0)
 	{ }
 
-	/** Virtual destructor. */
-	virtual ~FWmfMediaReadState() { }
-
 public:
 
 	/**
@@ -116,6 +113,14 @@ public:
 		return CurrentRefCount;
 	}
 
+private:
+
+	/** Hidden destructor (this class is reference counted). */
+	virtual ~FWmfMediaReadState()
+	{
+		check(RefCount == 0);
+	}
+
 private:
 
 	/** Number of bytes read. */
diff --git a/Engine/Plugins/Media/WmfMedia/Source/WmfMedia/Private/Wmf/WmfMediaSampler.h b/Engine/Plugins/Media/WmfMedia/Source/WmfMedia/Private/Wmf/WmfMediaSampler.h
--- a/Engine/Plugins/Media/WmfMedia/Source/WmfMedia/Private/Wmf/WmfMediaSampler.h
+++ b/Engine/Plugins/Media/WmfMedia/Source/WmfMedia/Private/Wmf/WmfMediaSampler.h
@@ -2,12 +2,18 @@
 
 #pragma once
 
-#include "../WmfMediaPrivate.h"
+#include "WmfMediaPrivate.h"
 
 #if WMFMEDIA_SUPPORTED_PLATFORM
 
+#include "CoreTypes.h"
+#include "Delegates/Delegate.h"
+#include "HAL/PlatformAtomics.h"
+#include "Misc/Timespan.h"
+
 #include "AllowWindowsPlatformTypes.h"
 
+
 /**
  * Clock sink events.
  */
@@ -33,9 +39,6 @@ public:
 		: RefCount(0)
 	{ }
 
-	/** Virtual destructor. */
-	virtual ~FWmfMediaSampler() { }
-
 public:
 
 	/** Get an event that gets fired when the sampler's presentation clock changed its state. */
@@ -56,11 +59,6 @@ public:
 
 	//~ IMFSampleGrabberSinkCallback interface
 
-	STDMETHODIMP_(ULONG) AddRef()
-	{
-		return FPlatformAtomics::InterlockedIncrement(&RefCount);
-	}
-
 	STDMETHODIMP OnClockPause(MFTIME SystemTime)
 	{
 		ClockEvent.Broadcast(EWmfMediaSamplerClockEvent::Paused);
@@ -111,6 +109,15 @@ public:
 		return S_OK;
 	}
 
+public:
+
+	//~ IUnknown interface
+
+	STDMETHODIMP_(ULONG) AddRef()
+	{
+		return FPlatformAtomics::InterlockedIncrement(&RefCount);
+	}
+
 #if _MSC_VER == 1900
 	#pragma warning(push)
 	#pragma warning(disable:4838)
@@ -143,6 +150,14 @@ public:
 		return CurrentRefCount;
 	}
 
+private:
+
+	/** Hidden destructor (this class is reference counted). */
+	virtual ~FWmfMediaSampler()
+	{
+		check(RefCount == 0);
+	}
+
 private:
 
 	/** Event that gets fired when the sampler's presentation clock changed state. */
@@ -158,4 +173,4 @@ private:
 
 #include "HideWindowsPlatformTypes.h"
 
-#endif
+#endif //WMFMEDIA_SUPPORTED_PLATFORM
