Copying //UE4/Dev-Sequencer to //UE4/Main (Source: //UE4/Dev-Sequencer @ 3617642)

diff --git a/Engine/Plugins/Media/WmfMedia/Source/WmfMedia/Private/Wmf/WmfMediaSession.cpp b/Engine/Plugins/Media/WmfMedia/Source/WmfMedia/Private/Wmf/WmfMediaSession.cpp
--- a/Engine/Plugins/Media/WmfMedia/Source/WmfMedia/Private/Wmf/WmfMediaSession.cpp
+++ b/Engine/Plugins/Media/WmfMedia/Source/WmfMedia/Private/Wmf/WmfMediaSession.cpp
@@ -1,56 +1,120 @@
 // Copyright 1998-2017 Epic Games, Inc. All Rights Reserved.
 
 #include "WmfMediaSession.h"
-#include "WmfMediaUtils.h"
-#include "Misc/ScopeLock.h"
 
 #if WMFMEDIA_SUPPORTED_PLATFORM
 
+#include "IMediaEventSink.h"
+#include "Misc/ScopeLock.h"
+
+#include "WmfMediaUtils.h"
+
 #include "AllowWindowsPlatformTypes.h"
 
 
+/* Local helpers
+ *****************************************************************************/
+
+namespace WmfMediaSession
+{
+	/** Get the human readable string representation of a media player state. */
+	const TCHAR* StateToString(EMediaState State)
+	{
+		switch (State)
+		{
+		case EMediaState::Closed: return TEXT("Closed");
+		case EMediaState::Error: return TEXT("Error");
+		case EMediaState::Paused: return TEXT("Paused");
+		case EMediaState::Preparing: return TEXT("Preparing");
+		case EMediaState::Stopped: return TEXT("Stopped");
+		default: return TEXT("Unknown");
+		}
+	}
+}
+
+
 /* FWmfMediaSession structors
  *****************************************************************************/
 
-FWmfMediaSession::FWmfMediaSession(EMediaState InState)
-	: Buffering(false)
-	, CanScrub(false)
+FWmfMediaSession::FWmfMediaSession()
+	: CanScrub(false)
 	, Capabilities(0)
-	, ChangeRequested(false)
+	, CurrentDuration(FTimespan::Zero())
 	, CurrentRate(0.0f)
-	, CurrentState(InState)
-	, Duration(FTimespan::Zero())
-	, Looping(false)
+	, CurrentState(EMediaState::Closed)
+	, LastTime(FTimespan::Zero())
+	, PendingChange(false)
 	, RefCount(0)
 	, RequestedRate(0.0f)
 	, RequestedTime(FTimespan::MinValue())
-	, StateChangePending(false)
+	, ShouldLoop(false)
+	, Status(EMediaStatus::None)
+{ }
+
+
+FWmfMediaSession::~FWmfMediaSession()
 {
-	check((InState == EMediaState::Closed) || (InState == EMediaState::Preparing));
+	check(RefCount == 0);
+	Shutdown();
+
+	UE_LOG(LogWmfMedia, Verbose, TEXT("Session %llx: Destroyed"), this);
 }
 
 
-FWmfMediaSession::FWmfMediaSession(const FTimespan& InDuration)
-	: Buffering(false)
-	, Capabilities(0)
-	, ChangeRequested(false)
-	, CurrentRate(0.0f)
-	, Duration(InDuration)
-	, Looping(false)
-	, RefCount(0)
-	, RequestedTime(FTimespan::MinValue())
-	, RequestedRate(0.0f)
-	, StateChangePending(false)
+/* FWmfMediaSession interface
+ *****************************************************************************/
+
+void FWmfMediaSession::GetEvents(TArray<EMediaEvent>& OutEvents)
 {
-	CurrentState = EMediaState::Error;
-	
+	EMediaEvent Event;
+
+	while (DeferredEvents.Dequeue(Event))
+	{
+		OutEvents.Add(Event);
+	}
+}
+
+
+bool FWmfMediaSession::Initialize(bool LowLatency)
+{
+	Shutdown();
+
+	// create session attributes
+	TComPtr<IMFAttributes> Attributes;
+	{
+		HRESULT Result = ::MFCreateAttributes(&Attributes, 2);
+
+		if (FAILED(Result))
+		{
+			UE_LOG(LogWmfMedia, Error, TEXT("Failed to create media session attributes (%s)"), *WmfMedia::ResultToString(Result));
+			return false;
+		}
+	}
+
+	if (LowLatency)
+	{
+		if (FWindowsPlatformMisc::VerifyWindowsVersion(6, 2))
+		{
+			HRESULT Result = Attributes->SetUINT32(MF_LOW_LATENCY, TRUE);
+
+			if (FAILED(Result))
+			{
+				UE_LOG(LogWmfMedia, Warning, TEXT("Failed to set low latency session attribute (%s)"), *WmfMedia::ResultToString(Result));
+			}
+		}
+		else
+		{
+			UE_LOG(LogWmfMedia, Warning, TEXT("Low latency media processing requires Windows 8 or newer"));
+		}
+	}
+
 	// create media session
-	HRESULT Result = ::MFCreateMediaSession(NULL, &MediaSession);
+	HRESULT Result = ::MFCreateMediaSession(Attributes, &MediaSession);
 
 	if (FAILED(Result))
 	{
 		UE_LOG(LogWmfMedia, Error, TEXT("Failed to create media session (%s)"), *WmfMedia::ResultToString(Result));
-		return;
+		return false;
 	}
 
 	// start media event processing
@@ -59,62 +123,142 @@ FWmfMediaSession::FWmfMediaSession(const FTimespan& InDuration)
 	if (FAILED(Result))
 	{
 		UE_LOG(LogWmfMedia, Error, TEXT("Failed to start media session event processing (%s)"), *WmfMedia::ResultToString(Result));
-		return;
+		return false;
+	}
+
+	CurrentState = EMediaState::Preparing;
+
+	return true;
+}
+
+
+bool FWmfMediaSession::SetTopology(const TComPtr<IMFTopology>& InTopology, FTimespan InDuration)
+{
+	if (MediaSession == NULL)
+	{
+		return false;
 	}
 
-	// get presentation clock, if available
-	TComPtr<IMFClock> Clock;
+	UE_LOG(LogWmfMedia, Verbose, TEXT("Session: %llx: Setting topology %llx (duration = %s)"), this, InTopology.Get(), *InDuration.ToString());
 
-	if (SUCCEEDED(MediaSession->GetClock(&Clock)))
+	FScopeLock Lock(&CriticalSection);
+
+	if (CurrentState == EMediaState::Preparing)
+	{
+		// media source resolved
+		if (InTopology != NULL)
+		{
+			// at least one track selected
+			HRESULT Result = MediaSession->SetTopology(MFSESSION_SETTOPOLOGY_IMMEDIATE, InTopology);
+
+			if (FAILED(Result))
+			{
+				UE_LOG(LogWmfMedia, Error, TEXT("Failed to set topology: %s"), *WmfMedia::ResultToString(Result));
+				
+				CurrentState = EMediaState::Error;
+				DeferredEvents.Enqueue(EMediaEvent::MediaOpenFailed);
+			}
+			else
+			{
+				// do nothing (Preparing state will exit in MESessionTopologyStatus event)
+			}
+		}
+		else
+		{
+			// no tracks selected
+			CurrentState = EMediaState::Stopped;
+			DeferredEvents.Enqueue(EMediaEvent::MediaOpened);
+		}
+	}
+	else
 	{
-		Clock->QueryInterface(IID_PPV_ARGS(&PresentationClock));
+		// topology changed during playback, i.e. track switching
+		if (PendingChange)
+		{
+			RequestedTopology = InTopology;
+		}
+		else
+		{
+			CommitTopology(InTopology);
+		}
 	}
 
-	// get rate control and support, if available
-	if (SUCCEEDED(::MFGetService(MediaSession, MF_RATE_CONTROL_SERVICE, IID_PPV_ARGS(&RateControl))) &&
-		SUCCEEDED(::MFGetService(MediaSession, MF_RATE_CONTROL_SERVICE, IID_PPV_ARGS(&RateSupport))))
+	CurrentDuration = InDuration;
+
+	return true;
+}
+
+
+void FWmfMediaSession::Shutdown()
+{
+	if (MediaSession == NULL)
 	{
-		CanScrub = SUCCEEDED(RateSupport->IsRateSupported(TRUE, 0, NULL));
+		return;
 	}
 
-	CurrentState = EMediaState::Stopped;
-	RequestedState = CurrentState;
+	FScopeLock Lock(&CriticalSection);
+
+	DiscardPendingState();
+
+	MediaSession->Close();
+	MediaSession->Shutdown();
+	MediaSession.Reset();
+
+	CurrentTopology.Reset();
+	PresentationClock.Reset();
+	RateControl.Reset();
+	RateSupport.Reset();
+
+	CanScrub = false;
+	Capabilities = 0;
+	CurrentDuration = FTimespan::Zero();
+	CurrentRate = 0.0f;
+	CurrentState = EMediaState::Closed;
+	LastTime = FTimespan::Zero();
+	RequestedRate = 0.0f;
+	Status = EMediaStatus::None;
+	ThinnedRates.Empty();
+	UnthinnedRates.Empty();
 }
 
 
-/* FWmfMediaSession interface
+/* IMediaControls interface
  *****************************************************************************/
 
-bool FWmfMediaSession::SetTopology(const TComPtr<IMFTopology>& NewTopology)
+bool FWmfMediaSession::CanControl(EMediaControl Control) const
 {
 	if (MediaSession == NULL)
 	{
 		return false;
 	}
 
-	if (CurrentState == EMediaState::Playing)
+	if (Control == EMediaControl::Pause)
 	{
-		HRESULT Result = MediaSession->SetTopology(0, NewTopology);
+		return ((CurrentState == EMediaState::Playing) && (((Capabilities & MFSESSIONCAP_PAUSE) != 0) || UnthinnedRates.Contains(0.0f)));
+	}
 
-		if (FAILED(Result))
-		{
-			UE_LOG(LogWmfMedia, Error, TEXT("Failed to change topology in media session (%s)"), *WmfMedia::ResultToString(Result));
-			return false;
-		}
+	if (Control == EMediaControl::Resume)
+	{
+		return ((CurrentState != EMediaState::Playing) && UnthinnedRates.Contains(1.0f));
 	}
 
-	Topology = NewTopology;
+	if (Control == EMediaControl::Scrub)
+	{
+		return CanScrub;
+	}
 
-	return true;
-}
+	if (Control == EMediaControl::Seek)
+	{
+		return (((Capabilities & MFSESSIONCAP_SEEK) != 0) && (CurrentDuration > FTimespan::Zero()));
+	}
 
+	return false;
+}
 
-/* IMediaControls interface
- *****************************************************************************/
 
 FTimespan FWmfMediaSession::GetDuration() const
 {
-	return Duration;
+	return CurrentDuration;
 }
 
 
@@ -126,67 +270,50 @@ float FWmfMediaSession::GetRate() const
 
 EMediaState FWmfMediaSession::GetState() const
 {
-	if (Buffering)
-	{
-		return EMediaState::Preparing;
-	}
-
-	if ((CurrentState == EMediaState::Playing) && FMath::IsNearlyZero(CurrentRate))
-	{
-		// scrubbing counts as paused
-		return EMediaState::Paused;
-	}
-
 	return CurrentState;
 }
 
 
-TRange<float> FWmfMediaSession::GetSupportedRates(EMediaPlaybackDirections Direction, bool Unthinned) const
+EMediaStatus FWmfMediaSession::GetStatus() const
 {
-	if (RateSupport == NULL)
-	{
-		return TRange<float>(0.0f);
-	}
-
-	float MaxRate = 0.0f;
-	float MinRate = 0.0f;
-
-	switch (Direction)
-	{
-	case EMediaPlaybackDirections::Forward:
-		RateSupport->GetSlowestRate(MFRATE_FORWARD, !Unthinned, &MinRate);
-		RateSupport->GetFastestRate(MFRATE_FORWARD, !Unthinned, &MaxRate);
-		break;
+	return Status;
+}
 
-	case EMediaPlaybackDirections::Reverse:
-		RateSupport->GetSlowestRate(MFRATE_REVERSE, !Unthinned, &MinRate);
-		RateSupport->GetFastestRate(MFRATE_REVERSE, !Unthinned, &MaxRate);
-		break;
-	}
 
-	return TRange<float>(MinRate, MaxRate);
+TRangeSet<float> FWmfMediaSession::GetSupportedRates(EMediaRateThinning Thinning) const
+{
+	return (Thinning == EMediaRateThinning::Thinned) ? ThinnedRates : UnthinnedRates;
 }
 
 
 FTimespan FWmfMediaSession::GetTime() const
 {
-	if (PresentationClock == NULL)
+	MFCLOCK_STATE ClockState;
+
+	if ((PresentationClock == NULL) || FAILED(PresentationClock->GetState(0, &ClockState)) || (ClockState == MFCLOCK_STATE_INVALID))
 	{
-		return FTimespan::Zero();
+		return FTimespan::Zero(); // topology not initialized, or clock not started yet
 	}
 
-	if (RequestedTime >= FTimespan::Zero())
+	if (ClockState == MFCLOCK_STATE_STOPPED)
 	{
-		return RequestedTime;
+		return LastTime; // WMF always reports zero when stopped
 	}
 
-	return GetInternalPosition();
+	MFTIME ClockTime;
+
+	if (FAILED(PresentationClock->GetTime(&ClockTime)))
+	{
+		return FTimespan::Zero();
+	}
+
+	return FTimespan(ClockTime);
 }
 
 
 bool FWmfMediaSession::IsLooping() const
 {
-	return Looping;
+	return ShouldLoop;
 }
 
 
@@ -197,20 +324,40 @@ bool FWmfMediaSession::Seek(const FTimespan& Time)
 		return false;
 	}
 
+	UE_LOG(LogWmfMedia, Verbose, TEXT("Session %llx: Seeking to %s"), this, *Time.ToString());
+
+	FScopeLock Lock(&CriticalSection);
+
+	// validate seek
 	if ((CurrentState == EMediaState::Closed) || (CurrentState == EMediaState::Error))
 	{
+		UE_LOG(LogWmfMedia, Warning, TEXT("Cannot seek while closed or in error state"));
+		return false;
+	}
+
+	if ((Time < FTimespan::Zero()) || (Time > CurrentDuration))
+	{
+		UE_LOG(LogWmfMedia, Warning, TEXT("Invalid seek time %s (media duration is %s)"), *Time.ToString(), *CurrentDuration.ToString());
 		return false;
 	}
 
-	RequestedTime = Time;
+	// wait for pending changes to complete
+	if (PendingChange)
+	{
+		UE_LOG(LogWmfMedia, Verbose, TEXT("Session %llx: Requesting seek after pending command"), this);
+
+		RequestedTime = Time;
 
-	return ChangeState();
+		return true;
+	}
+
+	return CommitTime(Time);
 }
 
 
-bool FWmfMediaSession::SetLooping(bool InLooping)
+bool FWmfMediaSession::SetLooping(bool Looping)
 {
-	Looping = InLooping;
+	ShouldLoop = Looping;
 
 	return true;
 }
@@ -223,97 +370,44 @@ bool FWmfMediaSession::SetRate(float Rate)
 		return false;
 	}
 
-	if (!SupportsRate(Rate, false))
-	{
-		return false;
-	}
-
-	RequestedRate = Rate;
+	FScopeLock Lock(&CriticalSection);
 
-	if (FMath::IsNearlyZero(Rate))
+	if (Rate == CurrentRate)
 	{
-		if (!SupportsScrubbing())
-		{
-			RequestedState = EMediaState::Paused;
-		}
+		return true; // rate already set
 	}
-	else
-	{
-		if (CurrentState != EMediaState::Playing)
-		{
-			HRESULT Result = MediaSession->SetTopology(0, Topology);
 
-			if (FAILED(Result))
-			{
-				UE_LOG(LogWmfMedia, Error, TEXT("Failed to set topology in media session (%s)"), *WmfMedia::ResultToString(Result));
-				return false;
-			}
-		}
+	EMediaRateThinning Thinning;
 
-		RequestedState = EMediaState::Playing;
-	}
-
-	return ChangeState();
-}
-
-
-bool FWmfMediaSession::SupportsRate(float Rate, bool Unthinned) const
-{
-	if ((CurrentState == EMediaState::Closed) || (CurrentState == EMediaState::Error))
+	// validate rate
+	if (UnthinnedRates.Contains(Rate))
 	{
-		return false;
+		UE_LOG(LogWmfMedia, Verbose, TEXT("Session %llx: Setting rate from to %f to %f (unthinned)"), this, CurrentRate, Rate);
+		Thinning = EMediaRateThinning::Unthinned;
 	}
-
-	if (Rate == 1.0f)
+	else if (ThinnedRates.Contains(Rate))
 	{
-		return true;
+		UE_LOG(LogWmfMedia, Verbose, TEXT("Session %llx: Setting rate from to %f to %f (thinned)"), this, CurrentRate, Rate);
+		Thinning = EMediaRateThinning::Thinned;
 	}
-
-	if ((RateControl == NULL) || (RateSupport == NULL))
+	else
 	{
+		UE_LOG(LogWmfMedia, Warning, TEXT("The rate %f is not supported"), Rate);
 		return false;
 	}
 
-	if (Unthinned)
-	{
-		return SUCCEEDED(RateSupport->IsRateSupported(FALSE, Rate, NULL));
-	}
-
-	return SUCCEEDED(RateSupport->IsRateSupported(TRUE, Rate, NULL));
-}
-
-
-bool FWmfMediaSession::SupportsScrubbing() const
-{
-	return CanScrub;
-}
-
-
-bool FWmfMediaSession::SupportsSeeking() const
-{
-	return (((Capabilities & MFSESSIONCAP_SEEK) != 0) && (Duration > FTimespan::Zero()));
-}
-
-
-/* FWmfMediaSession interface
- *****************************************************************************/
-
-bool FWmfMediaSession::SetState(EMediaState NewState)
-{
-	if (MediaSession == NULL)
+	// wait for pending changes to complete
+	if (PendingChange)
 	{
-		return false;
-	}
+		UE_LOG(LogWmfMedia, Verbose, TEXT("Session %llx: Requesting rate change after pending command"), this);
 
-	RequestedState = NewState;
+		RequestedRate = Rate;
+		RequestedThinning = Thinning;
 
-	if ((NewState == EMediaState::Closed) || (NewState == EMediaState::Stopped))
-	{
-		RequestedTime = FTimespan::Zero();
-		RequestedRate = 0.0f;
+		return true;
 	}
 
-	return ChangeState();
+	return CommitRate(Rate, Thinning);
 }
 
 
@@ -328,125 +422,394 @@ STDMETHODIMP_(ULONG) FWmfMediaSession::AddRef()
 
 STDMETHODIMP FWmfMediaSession::GetParameters(unsigned long*, unsigned long*)
 {
-	return E_NOTIMPL;
+	return E_NOTIMPL; // default behavior
 }
 
 
 STDMETHODIMP FWmfMediaSession::Invoke(IMFAsyncResult* AsyncResult)
 {
+	FScopeLock Lock(&CriticalSection);
+
+	if (MediaSession == NULL)
+	{
+		return S_OK;
+	}
+
 	// get event
 	TComPtr<IMFMediaEvent> Event;
-	
-	if (FAILED(MediaSession->EndGetEvent(AsyncResult, &Event)))
 	{
-		UE_LOG(LogWmfMedia, Verbose, TEXT("Failed to get event"));
-		return S_OK;
+		const HRESULT Result = MediaSession->EndGetEvent(AsyncResult, &Event);
+
+		if (FAILED(Result))
+		{
+			UE_LOG(LogWmfMedia, Verbose, TEXT("Session %llx: Failed to get event: %s"), this, *WmfMedia::ResultToString(Result));
+			return S_OK;
+		}
 	}
 
 	MediaEventType EventType = MEUnknown;
-	
-	if (FAILED(Event->GetType(&EventType)))
 	{
-		UE_LOG(LogWmfMedia, Verbose, TEXT("Failed to get session event type"));
-		return S_OK;
-	}
+		const HRESULT Result = Event->GetType(&EventType);
 
-	// process event
-	HRESULT EventResult = S_OK;
+		if (FAILED(Result))
+		{
+			UE_LOG(LogWmfMedia, Verbose, TEXT("Session %llx: Failed to get session event type: %s"), this, *WmfMedia::ResultToString(Result));
+			return S_OK;
+		}
+	}
 
-	if (FAILED(Event->GetStatus(&EventResult)))
+	HRESULT EventStatus = S_FALSE;
 	{
-		UE_LOG(LogWmfMedia, Verbose, TEXT("Failed to get event result"));
+		const HRESULT Result = Event->GetStatus(&EventStatus);
+
+		if (FAILED(Result))
+		{
+			UE_LOG(LogWmfMedia, Verbose, TEXT("Session %llx: Failed to get event status: %s"), this, *WmfMedia::ResultToString(Result));
+		}
 	}
 
-	UE_LOG(LogWmfMedia, Verbose, TEXT("Media session event: %s: %s"), *WmfMedia::SessionEventToString(EventType), *WmfMedia::ResultToString(EventResult));
+	UE_LOG(LogWmfMedia, Verbose, TEXT("Session %llx: Event [%s] (%s)"), this, *WmfMedia::MediaEventToString(EventType), *WmfMedia::ResultToString(EventStatus));
 
-	FScopeLock ScopeLock(&CriticalSection);
+	// process event
+	switch (EventType)
 	{
-		switch (EventType)
+	case MEBufferingStarted:
+		Status |= EMediaStatus::Buffering;
+		DeferredEvents.Enqueue(EMediaEvent::MediaBuffering);
+		break;
+
+	case MEBufferingStopped:
+		Status &= ~EMediaStatus::Buffering;
+		break;
+
+	case MEError:
+		UE_LOG(LogWmfMedia, Error, TEXT("An error occurred in the media session: %s"), *WmfMedia::ResultToString(EventStatus));
+		CurrentState = EMediaState::Error;
+		DiscardPendingState();
+		MediaSession->Close();
+		break;
+
+	case MEReconnectEnd:
+		Status &= ~EMediaStatus::Connecting;
+		break;
+
+	case MEReconnectStart:
+		Status |= EMediaStatus::Connecting;
+		DeferredEvents.Enqueue(EMediaEvent::MediaConnecting);
+		break;
+
+	case MESessionCapabilitiesChanged:
+		Capabilities = ::MFGetAttributeUINT32(Event, MF_EVENT_SESSIONCAPS, Capabilities);
+		break;
+
+	case MESessionClosed:
+		Capabilities = 0;
+		LastTime = FTimespan::Zero();
+		break;
+
+	case MESessionEnded:
+		DeferredEvents.Enqueue(EMediaEvent::PlaybackEndReached);
+
+		if (ShouldLoop)
 		{
-		case MEBufferingStarted:
-			Buffering = true;
-			break;
+			// loop back to beginning/end
+			RequestedTime = (CurrentRate < 0.0f) ? CurrentDuration : FTimespan::Zero();
+			UpdatePendingState();
+		}
+		else
+		{
+			CurrentRate = RequestedRate = 0.0f;
+			CurrentState = EMediaState::Stopped;
+			LastTime = FTimespan::Zero();
 
-		case MEBufferingStopped:
-			Buffering = false;
-			break;
+			if (RateControl != NULL)
+			{
+				RateControl->SetRate(TRUE, 0.0f);
+			}
+		}
+		break;
 
-		case MEError:
-			UE_LOG(LogWmfMedia, Error, TEXT("An error occured in the media session: %s"), *WmfMedia::ResultToString(EventResult));
-			UpdateState(EMediaState::Error);
-			break;
+	case MESessionPaused:
+		if (SUCCEEDED(EventStatus))
+		{
+			CurrentState = EMediaState::Paused;
+			DeferredEvents.Enqueue(EMediaEvent::PlaybackSuspended);
+			UpdatePendingState();
+		}
+		else
+		{
+			DiscardPendingState();
+		}
+		break;
 
-		case MESessionCapabilitiesChanged:
-			Capabilities = ::MFGetAttributeUINT32(Event, MF_EVENT_SESSIONCAPS, Capabilities);
-			break;
+	case MESessionRateChanged:
+		if (SUCCEEDED(EventStatus))
+		{
+			// cache current rate
+			PROPVARIANT Value;
+			::PropVariantInit(&Value);
 
-		case MESessionClosed:
-			UpdateState(EMediaState::Closed);
-			Capabilities = 0;
-			break;
+			const HRESULT Result = Event->GetValue(&Value);
+
+			if (SUCCEEDED(Result) && (Value.vt == VT_R4))
+			{
+				CurrentRate = Value.fltVal;
+			}
+		}
+		else if (RateControl != NULL)
+		{
+			BOOL Thin = FALSE;
+			RateControl->GetRate(&Thin, &CurrentRate);
+		}
+
+		RequestedRate = CurrentRate;
+		UpdatePendingState();
+		break;
+
+	case MESessionScrubSampleComplete:
+		DeferredEvents.Enqueue(EMediaEvent::SeekCompleted);
+		UpdatePendingState();
+		break;
 
-		case MESessionEnded:
-			if (Looping)
+	case MESessionStarted:
+		if (SUCCEEDED(EventStatus))
+		{
+			if (CurrentState == EMediaState::Playing)
 			{
-				RequestedTime = (CurrentRate < 0.0f) ? Duration : FTimespan::Zero();
-				ChangeState();
+				DeferredEvents.Enqueue(EMediaEvent::SeekCompleted);
+			}
+
+			CurrentState = (CurrentRate == 0.0f) ? EMediaState::Paused : EMediaState::Playing;
+			DeferredEvents.Enqueue(EMediaEvent::PlaybackResumed);
+
+			if (CurrentState == EMediaState::Paused)
+			{
+				// wait for MESessionScrubSampleComplete
 			}
 			else
 			{
-				RequestedState = EMediaState::Stopped;
-				UpdateState(EMediaState::Stopped);
+				UpdatePendingState();
 			}
-			break;
+		}
+		else
+		{
+			DiscardPendingState();
+		}
+		break;
 
-		case MESessionPaused:
-			UpdateState(EMediaState::Paused);
-			break;
+	case MESessionStopped:
+		if (SUCCEEDED(EventStatus))
+		{
+			CurrentState = EMediaState::Stopped;
+			DeferredEvents.Enqueue(EMediaEvent::PlaybackSuspended);
+			UpdatePendingState();
+		}
+		else
+		{
+			DiscardPendingState();
+		}
+		break;
+
+	case MESessionTopologySet:
+		if (SUCCEEDED(EventStatus))
+		{
+			if (CurrentState != EMediaState::Preparing)
+			{
+				// track switching successful
+				RequestedTopology.Reset();
+				UpdatePendingState();
+			}
+		}
+		else
+		{
+			if (CurrentState == EMediaState::Preparing)
+			{
+				// new media has incorrect topology
+				CurrentState = EMediaState::Error;
+				DeferredEvents.Enqueue(EMediaEvent::MediaOpenFailed);
+			}
+
+			DiscardPendingState();
+		}
+		break;
 
-		case MESessionRateChanged:
-			// recover active playback rate if rate change failed
-			if (FAILED(EventResult) && (CurrentRate == RequestedRate))
+	case MESessionTopologyStatus:
+		if (SUCCEEDED(EventStatus))
+		{
+			// check if new topology is ready
+			UINT32 TopologyStatus = MF_TOPOSTATUS_INVALID;
 			{
-				PROPVARIANT Value;
-				::PropVariantInit(&Value);
+				const HRESULT Result = Event->GetUINT32(MF_EVENT_TOPOLOGY_STATUS, (UINT32*)&TopologyStatus);
 
-				if (SUCCEEDED(Event->GetValue(&Value)) && (Value.vt == VT_R4))
+				if (FAILED(Result))
 				{
-					CurrentRate = Value.fltVal;
+					UE_LOG(LogWmfMedia, Verbose, TEXT("Session %llx: Failed to get topology status: %s"), this, *WmfMedia::ResultToString(Result));
 				}
-
-				RequestedRate = CurrentRate;
 			}
 
-		case MESessionScrubSampleComplete:
-			if (GetInternalPosition() == RequestedTime)
+			// initialize new topology
+			if (TopologyStatus == MF_TOPOSTATUS_READY)
 			{
-				RequestedTime = FTimespan::MinValue();
+				UE_LOG(LogWmfMedia, Verbose, TEXT("Session %llx: Topology ready"), this);
+
+				HRESULT Result = MediaSession->GetFullTopology(MFSESSION_GETFULLTOPOLOGY_CURRENT, 0, &CurrentTopology);
+
+				if (SUCCEEDED(Result))
+				{
+					// get presentation clock, if available
+					TComPtr<IMFClock> Clock;
+
+					Result = MediaSession->GetClock(&Clock);
+
+					if (FAILED(Result))
+					{
+						UE_LOG(LogWmfMedia, Log, TEXT("Session clock unavailable (%s)"), *WmfMedia::ResultToString(Result));
+					}
+					else
+					{
+						Result = Clock->QueryInterface(IID_PPV_ARGS(&PresentationClock));
+
+						if (FAILED(Result))
+						{
+							UE_LOG(LogWmfMedia, Log, TEXT("Presentation clock unavailable (%s)"), *WmfMedia::ResultToString(Result));
+						}
+						else
+						{
+							UE_LOG(LogWmfMedia, Verbose, TEXT("Session %llx: Presentation clock ready"), this);
+						}
+					}
+
+					// get rate control & rate support, if available
+					Result = ::MFGetService(MediaSession, MF_RATE_CONTROL_SERVICE, IID_PPV_ARGS(&RateControl));
+
+					if (FAILED(Result))
+					{
+						UE_LOG(LogWmfMedia, Log, TEXT("Rate control service unavailable (%s)"), *WmfMedia::ResultToString(Result));
+					}
+					else
+					{
+						UE_LOG(LogWmfMedia, Verbose, TEXT("Session %llx: Rate control ready"), this);
+					}
+
+					Result = ::MFGetService(MediaSession, MF_RATE_CONTROL_SERVICE, IID_PPV_ARGS(&RateSupport));
+
+					if (FAILED(Result))
+					{
+						UE_LOG(LogWmfMedia, Log, TEXT("Rate support service unavailable (%s)"), *WmfMedia::ResultToString(Result));
+					}
+					else
+					{
+						UE_LOG(LogWmfMedia, Verbose, TEXT("Session %llx: Rate support ready"), this);
+					}
+
+					// cache rate control properties
+					if (RateSupport.IsValid())
+					{
+						CanScrub = SUCCEEDED(RateSupport->IsRateSupported(TRUE, 0.0f, NULL));
+
+						float MaxRate = 0.0f;
+						float MinRate = 0.0f;
+
+						if (SUCCEEDED(RateSupport->GetSlowestRate(MFRATE_FORWARD, TRUE, &MinRate)) &&
+							SUCCEEDED(RateSupport->GetFastestRate(MFRATE_FORWARD, TRUE, &MaxRate)))
+						{
+							ThinnedRates.Add(TRange<float>::Inclusive(MinRate, MaxRate));
+						}
+
+						if (SUCCEEDED(RateSupport->GetSlowestRate(MFRATE_REVERSE, TRUE, &MinRate)) &&
+							SUCCEEDED(RateSupport->GetFastestRate(MFRATE_REVERSE, TRUE, &MaxRate)))
+						{
+							ThinnedRates.Add(TRange<float>::Inclusive(MaxRate, MinRate));
+						}
+
+						if (SUCCEEDED(RateSupport->GetSlowestRate(MFRATE_FORWARD, FALSE, &MinRate)) &&
+							SUCCEEDED(RateSupport->GetFastestRate(MFRATE_FORWARD, FALSE, &MaxRate)))
+						{
+							UnthinnedRates.Add(TRange<float>::Inclusive(MinRate, MaxRate));
+						}
+
+						if (SUCCEEDED(RateSupport->GetSlowestRate(MFRATE_REVERSE, FALSE, &MinRate)) &&
+							SUCCEEDED(RateSupport->GetFastestRate(MFRATE_REVERSE, FALSE, &MaxRate)))
+						{
+							UnthinnedRates.Add(TRange<float>::Inclusive(MaxRate, MinRate));
+						}
+					}
+					else
+					{
+						CanScrub = false;
+						ThinnedRates.Empty();
+						UnthinnedRates.Empty();
+					}
+
+					// new media opened successfully
+					if (CurrentState == EMediaState::Preparing)
+					{
+						CurrentState = EMediaState::Stopped;
+						DeferredEvents.Enqueue(EMediaEvent::MediaOpened);
+
+						UpdatePendingState();
+					}
+				}
+				else
+				{
+					UE_LOG(LogWmfMedia, Verbose, TEXT("Session %llx: Failed to get full topology: %s"), this, *WmfMedia::ResultToString(Result));
+
+					if (CurrentState == EMediaState::Preparing)
+					{
+						// incorrect topology for new media
+						CurrentState = EMediaState::Error;
+						DeferredEvents.Enqueue(EMediaEvent::MediaOpenFailed);
+					}
+
+					DiscardPendingState();
+				}
+
+				LastTime = FTimespan::Zero();
 			}
-			break;
+		}
+		else
+		{
+			if (CurrentState == EMediaState::Preparing)
+			{
+				UE_LOG(LogWmfMedia, Error, TEXT("An error occured when preparing the topology"));
 
-		case MESessionStarted:
-			UpdateState(EMediaState::Playing);
-			break;
+				// an error occurred in the topology
+				CurrentState = EMediaState::Error;
+				DeferredEvents.Enqueue(EMediaEvent::MediaOpenFailed);
+			}
 
-		case MESessionStopped:
-			UpdateState(EMediaState::Stopped);
-			break;
+			DiscardPendingState();
 		}
+		break;
 
-		// bubble up event to player
-		SessionEvent.Broadcast(EventType);
+	default:
+		break; // unsupported event
+	}
 
-		// request the next event
-		if (FAILED(MediaSession->BeginGetEvent(this, NULL)))
+	// request next event
+	if ((EventType != MESessionClosed) && (CurrentState != EMediaState::Error) && (MediaSession != NULL))
+	{
+		const HRESULT Result = MediaSession->BeginGetEvent(this, NULL);
+
+		if (FAILED(Result))
 		{
-			UE_LOG(LogWmfMedia, Verbose, TEXT("Failed to request next session event"));
+			UE_LOG(LogWmfMedia, Warning, TEXT("Failed to request next session event; aborting playback: %s"), *WmfMedia::ResultToString(Result));
 
 			Capabilities = 0;
 			CurrentState = EMediaState::Error;
 		}
 	}
+
+	// log current state
+	UE_LOG(LogWmfMedia, VeryVerbose, TEXT("Session %llx: CurrentState: %s, CurrentRate: %f, RequestedRate: %f, CurrentTime: %s, RequestedTime: %s, Pending Change: %i"),
+		this,
+		WmfMediaSession::StateToString(CurrentState),
+		CurrentRate,
+		RequestedRate,
+		*GetTime().ToString(),
+		*RequestedTime.ToString(),
+		PendingChange
+	);
 	
 	return S_OK;
 }
@@ -486,199 +849,278 @@ STDMETHODIMP_(ULONG) FWmfMediaSession::Release()
 }
 
 
-/* FWmfVideoSession implementation
+/* FWmfMediaSession implementation
  *****************************************************************************/
 
-bool FWmfMediaSession::ChangeState()
+bool FWmfMediaSession::CommitRate(float Rate, EMediaRateThinning Thinning)
 {
-	// disallow state changes if session was closed or had an error
-	if (CurrentState == EMediaState::Closed)
-	{
-		return false;
-	}
+	check(MediaSession != NULL);
+	check(PendingChange == false);
 
-	FScopeLock ScopeLock(&CriticalSection);
+	// pause player for basic pausing or if scrubbing isn't supported
+	const bool ShouldPause = (Rate == 0.0f) && ((CurrentRate == 1.0f) || !CanControl(EMediaControl::Scrub));
+
+	if (ShouldPause)
 	{
-		if (CurrentState == EMediaState::Error)
+		if (CurrentState == EMediaState::Playing)
 		{
-			UE_LOG(LogWmfMedia, Verbose, TEXT("Closing media session on error"));
-			MediaSession->Close();
-			MediaSession->Shutdown();
-
-			return true;
-		}
+			const HRESULT Result = MediaSession->Pause();
 
-		// defer state change
-		if (StateChangePending)
-		{
-			ChangeRequested = true;
+			if (FAILED(Result))
+			{
+				UE_LOG(LogWmfMedia, Warning, TEXT("Failed to pause for rate change"));
+				return false;
+			}
 
-			return true;
-		}
+			UE_LOG(LogWmfMedia, Verbose, TEXT("Session %llx: Paused for rate change from %f to %f"), this, CurrentRate, Rate);
 
-		// close session
-		if (RequestedState == EMediaState::Closed)
-		{
-			UE_LOG(LogWmfMedia, Verbose, TEXT("Closing media session as requested"));
-			MediaSession->Close();
-			MediaSession->Shutdown();
+			// request deferred rate change
+			if (RateControl != NULL)
+			{
+				RequestedRate = Rate;
+				RequestedThinning = Thinning;
+			}
 
-			return true;
-		}
+			PendingChange = true;
 
-		// handle rate changes
-		if (UpdateRate())
-		{
 			return true;
 		}
+	}
 
-		// seek or scrub
-		if (RequestedTime >= FTimespan::Zero())
+	// most rate changes require transition to Stopped first
+	if (((Rate >= 0.0f) || (CurrentRate >= 0.0f)) && ((Rate <= 0.0f) || (CurrentRate <= 0.0f)))
+	{
+		if (!ShouldPause && (CurrentState != EMediaState::Stopped))
 		{
-			UE_LOG(LogWmfMedia, Verbose, TEXT("Starting playback to seek to %s"), *RequestedTime.ToString());
+			LastTime = GetTime();
+			const HRESULT Result = MediaSession->Stop();
 
-			PROPVARIANT StartPosition;
+			if (FAILED(Result))
 			{
-				StartPosition.vt = VT_I8;
-				StartPosition.hVal.QuadPart = RequestedTime.GetTicks();
+				UE_LOG(LogWmfMedia, Warning, TEXT("Failed to stop for rate change"));
+				return false;
 			}
 
-			HRESULT Result = MediaSession->Start(NULL, &StartPosition);
+			UE_LOG(LogWmfMedia, Verbose, TEXT("Session %llx: Stopped for rate change from %f to %f"), this, CurrentRate, Rate);
 
-			if (FAILED(Result))
+			// request deferred rate change
+			if (RequestedTime == FTimespan::MinValue())
 			{
-				UE_LOG(LogWmfMedia, Warning, TEXT("Failed to seek %i"), *WmfMedia::ResultToString(Result));
+				RequestedTime = LastTime;
 			}
-		}
 
-		// handle state changes
-		if (RequestedState == CurrentState)
-		{
+			RequestedRate = Rate;
+			RequestedThinning = Thinning;
+			PendingChange = true;
+
 			return true;
 		}
+	}
 
-		bool StateChangeFailed = false;
+	// determine restart time
+	FTimespan RestartTime;
 
-		if (RequestedState == EMediaState::Paused)
-		{
-			UE_LOG(LogWmfMedia, Verbose, TEXT("Pausing media session as requested"));
+	if (RequestedTime >= FTimespan::Zero())
+	{
+		RestartTime = RequestedTime;
+	}
+	else
+	{
+		RestartTime = GetTime();
+	}
+
+	if ((RestartTime == FTimespan::Zero()) && (Rate < 0.0f))
+	{
+		RestartTime = CurrentDuration; // loop to end
+	}
+	else if ((RestartTime == CurrentDuration) && (Rate > 0.0f))
+	{
+		RestartTime = FTimespan::Zero(); // loop to beginning
+	}
 
-			StateChangePending = true;
-			StateChangeFailed = FAILED(MediaSession->Pause());
-//			RequestedTime = GetTime();
+	// set desired rate
+	if (RateControl != NULL)
+	{
+		const TCHAR* ThinnedString = (Thinning == EMediaRateThinning::Thinned) ? TEXT("thinned") : TEXT("unthinned");
+		const HRESULT Result = RateControl->SetRate((Thinning == EMediaRateThinning::Thinned) ? TRUE : FALSE, Rate);
+
+		if (FAILED(Result))
+		{
+			UE_LOG(LogWmfMedia, Warning, TEXT("Failed to commit rate change from %f to %f (%s): %s"), CurrentRate, Rate, ThinnedString, *WmfMedia::ResultToString(Result));
+			return false;
 		}
-		else if (RequestedState == EMediaState::Playing)
+
+		UE_LOG(LogWmfMedia, Verbose, TEXT("Session %llx: Committed rate change from %f to %f (%s)"), this, CurrentRate, Rate, ThinnedString);
+
+		// request deferred restart
+		if (Rate != 0.0f)
 		{
-			PROPVARIANT StartPosition;
+			RequestedTime = RestartTime;
+			PendingChange = true;
+		}
 
-			if (RequestedTime >= FTimespan::Zero())
-			{
-				UE_LOG(LogWmfMedia, Verbose, TEXT("Starting playback at %s"), *RequestedTime.ToString());
+		return true;
+	}
 
-				StartPosition.vt = VT_I8;
-				StartPosition.hVal.QuadPart = RequestedTime.GetTicks();
-			}
-			else
-			{
-				UE_LOG(LogWmfMedia, Verbose, TEXT("Starting playback at beginning"));
-				::PropVariantInit(&StartPosition);
-			}
+	// restart playback
+	if (Rate != 0.0f)
+	{
+		CommitTime(RestartTime);
+	}
 
-			StateChangePending = true;
-			StateChangeFailed = FAILED(MediaSession->Start(NULL, &StartPosition));
-		}
-		else if (RequestedState == EMediaState::Stopped)
-		{
-			UE_LOG(LogWmfMedia, Verbose, TEXT("Stopping playback as requested"));
+	return true;
+}
 
-			StateChangePending = true;
-			StateChangeFailed = FAILED(MediaSession->Stop());
-		}
 
-		if (StateChangeFailed)
-		{
-			RequestedState = CurrentState;
+bool FWmfMediaSession::CommitTime(FTimespan Time)
+{
+	check(MediaSession != NULL);
+	check(PendingChange == false);
 
+	// reapply topology if needed
+#if 0
+	if ((CurrentState != EMediaState::Paused) && (CurrentState != EMediaState::Playing))
+	{
+		HRESULT Result = MediaSession->SetTopology(0, CurrentTopology);
+
+		if (FAILED(Result))
+		{
+			UE_LOG(LogWmfMedia, Error, TEXT("Failed to reapply topology (%s)"), *WmfMedia::ResultToString(Result));
 			return false;
 		}
+
+		UE_LOG(LogWmfMedia, Verbose, TEXT("Session %llx: Reapplied topology"), this);
 	}
+#endif
+
+	// seek to requested time
+	PROPVARIANT StartPosition;
+
+	if (!CanControl(EMediaControl::Seek) || (Time == FTimespan::MaxValue()))
+	{
+		StartPosition.vt = VT_EMPTY; // current time
+	}
+	else
+	{
+		StartPosition.vt = VT_I8;
+		StartPosition.hVal.QuadPart = Time.GetTicks();
+	}
+
+	HRESULT Result = MediaSession->Start(NULL, &StartPosition);
+
+	if (FAILED(Result))
+	{
+		UE_LOG(LogWmfMedia, Warning, TEXT("Failed to commit time %s: %s"), (Time == FTimespan::MaxValue()) ? TEXT("current") : *Time.ToString(), *WmfMedia::ResultToString(Result));
+		return false;
+	}
+
+	UE_LOG(LogWmfMedia, Verbose, TEXT("Session %llx: Commited time %s"), this, (Time == FTimespan::MaxValue()) ? TEXT("current") : *Time.ToString());
+
+	PendingChange = true;
 
 	return true;
 }
 
 
-bool FWmfMediaSession::UpdateRate()
+bool FWmfMediaSession::CommitTopology(IMFTopology* Topology)
 {
-	bool Result = false;
+	check(MediaSession != NULL);
+	check(PendingChange == false);
 
-	if (RequestedRate != CurrentRate)
+	if (CurrentState != EMediaState::Stopped)
 	{
-		if (CurrentState != EMediaState::Stopped)
-		{
-			UE_LOG(LogWmfMedia, Verbose, TEXT("Stopping playback for rate change"));
-
-			if (RequestedTime == FTimespan::MinValue())
-			{
-				RequestedTime = GetTime();
-			}
+		LastTime = GetTime();
 
-			MediaSession->Stop();
-			Result = true;
-		}
+		// topology change requires transition to Stopped; playback is resumed afterwards
+		HRESULT Result = MediaSession->Stop();
 
-		if (SUCCEEDED(RateControl->SetRate(SupportsRate(RequestedRate, true) ? FALSE : TRUE, RequestedRate)))
+		if (FAILED(Result))
 		{
-			UE_LOG(LogWmfMedia, Verbose, TEXT("Set rate from %f to %f"), CurrentRate, RequestedRate);
-			CurrentRate = RequestedRate;
+			UE_LOG(LogWmfMedia, Warning, TEXT("Failed to stop for topology change"));
+			return false;
 		}
-		else
+
+		UE_LOG(LogWmfMedia, Verbose, TEXT("Session %llx: Stopped for topology change"), this);
+
+		// request deferred playback restart
+		if (RequestedTime == FTimespan::MinValue())
 		{
-			UE_LOG(LogWmfMedia, Verbose, TEXT("Failed to set rate from %f to %f"), CurrentRate, RequestedRate);
-			RequestedRate = CurrentRate;
+			RequestedTime = LastTime;
 		}
-	}
 
-	return Result;
-}
+		RequestedTopology = Topology;
+		PendingChange = true;
 
+		return true;
+	}
 
-FTimespan FWmfMediaSession::GetInternalPosition() const
-{
-	MFTIME ClockTime;
+	// set new topology
+	HRESULT Result = MediaSession->SetTopology(MFSESSION_SETTOPOLOGY_IMMEDIATE, Topology);
 
-	if (FAILED(PresentationClock->GetTime(&ClockTime)))
+	if (SUCCEEDED(Result))
 	{
-		// topology not initialized yet
-		return FTimespan::Zero();
+		UE_LOG(LogWmfMedia, Verbose, TEXT("Session %llx: Committed topology"), this);
+	}
+	else
+	{
+		UE_LOG(LogWmfMedia, Warning, TEXT("Failed to set topology: %s"), *WmfMedia::ResultToString(Result));
 	}
 
-	return FTimespan(ClockTime);
+	PendingChange = true;
+
+	return true;
 }
 
 
-void FWmfMediaSession::UpdateState(EMediaState CompletedState)
+void FWmfMediaSession::DiscardPendingState()
 {
-	CurrentState = CompletedState;
+	PendingChange = false;
+	RequestedRate = CurrentRate;
+	RequestedTime = FTimespan::MinValue();
+	RequestedTopology.Reset();
+}
 
-	if (CompletedState == EMediaState::Playing)
+
+void FWmfMediaSession::UpdatePendingState()
+{
+	UE_LOG(LogWmfMedia, Verbose, TEXT("Session %llx: Updating pending state"), this);
+
+	PendingChange = false;
+
+	// commit pending topology changes
+	if (RequestedTopology != NULL)
 	{
-		RequestedTime = FTimespan::MinValue();
+		if (!CommitTopology(RequestedTopology))
+		{
+			RequestedTopology.Reset();
+		}
+
+		if (PendingChange)
+		{
+			return;
+		}
 	}
 
-	if (CompletedState == RequestedState)
+	// commit pending rate changes
+	if (RequestedRate != CurrentRate)
 	{
-		StateChangePending = false;
-		UpdateRate();
+		if (!CommitRate(RequestedRate, RequestedThinning))
+		{
+			RequestedRate = CurrentRate;
+		}
 
-		if (RequestedState == EMediaState::Stopped)
+		if (PendingChange)
 		{
-			CurrentRate = 0.0f;
-			RequestedRate = 0.0f;
+			return;
 		}
 	}
-	else
+
+	// commit pending seeks/restarts
+	if (RequestedTime >= FTimespan::Zero())
 	{
-		ChangeState();
+		CommitTime(RequestedTime);
+
+		RequestedTime = FTimespan::MinValue();
 	}
 }
 
diff --git a/Engine/Plugins/Media/WmfMedia/Source/WmfMedia/Private/Wmf/WmfMediaSession.h b/Engine/Plugins/Media/WmfMedia/Source/WmfMedia/Private/Wmf/WmfMediaSession.h
--- a/Engine/Plugins/Media/WmfMedia/Source/WmfMedia/Private/Wmf/WmfMediaSession.h
+++ b/Engine/Plugins/Media/WmfMedia/Source/WmfMedia/Private/Wmf/WmfMediaSession.h
@@ -2,15 +2,20 @@
 
 #pragma once
 
-#include "../WmfMediaPrivate.h"
-#include "IMediaControls.h"
+#include "WmfMediaPrivate.h"
 
 #if WMFMEDIA_SUPPORTED_PLATFORM
 
+#include "CoreTypes.h"
+#include "Containers/Queue.h"
+#include "HAL/CriticalSection.h"
+#include "IMediaControls.h"
+#include "IMediaEventSink.h"
+#include "Misc/Timespan.h"
+
 #include "AllowWindowsPlatformTypes.h"
 
-// forward declarations
-enum class EMediaPlaybackDirections;
+enum class EMediaEvent;
 
 
 /**
@@ -20,101 +25,88 @@ enum class EMediaPlaybackDirections;
  * immediately, such as seeking and playback rate changes. A media session may
  * generate events during playback that are then handled by this class.
  *
- * Windows Media Foundation also queues up most playback commands, which may have
- * undesired side effects, such as unresponsive or sluggish user interfaces. For
- * this reason, this helper class also implements a mechanism to manage pending
- * operations efficiently.
- *
- * @todo gmp: implement better command queuing.
+ * Windows Media Foundation has a number of odd quirks and problems that require
+ * special handling, such as certain state changes not being allowed, and some
+ * calls causing occasional deadlocks. The added complexity in the implementation
+ * of this class is for working around those issues.
  */
 class FWmfMediaSession
 	: public IMFAsyncCallback
 	, public IMediaControls
 {
+public:
+
+	/** Default constructor. */
+	FWmfMediaSession();
+
 public:
 
 	/**
-	 * Create a placeholder instance.
-	 *
-	 * This constructor is used by WmfMediaPlayer when no media
-	 * is open or when a media URL is currently being resolved.
+	 * Gets the session capabilities.
 	 *
-	 * @param InState Must be either Closed or Preparing.
+	 * @return Capabilities bit mask.
+	 * @see GetEvents
 	 */
-	FWmfMediaSession(EMediaState InState);
+	DWORD GetCapabilities() const
+	{
+		return Capabilities;
+	}
 
 	/**
-	 * Creates and initializes a new instance.
+	 * Gets all deferred player events.
 	 *
-	 * @param InDuration The duration of the media.
+	 * @param OutEvents Will contain the events.
+	 * @see GetCapabilities
 	 */
-	FWmfMediaSession(const FTimespan& InDuration);
+	void GetEvents(TArray<EMediaEvent>& OutEvents);
 
-	/** Virtual destructor. */
-	virtual ~FWmfMediaSession() { }
-
-public:
+	/**
+	 * Initialize the media session.
+	 *
+	 * @param LowLatency Whether to enable low latency processing.
+	 * @see SetTopolgy, Shutdown
+	 */
+	bool Initialize(bool LowLatency);
 
 	/**
-	 * Set the given playback topology.
+	 * Set the playback topology to be used by this session.
 	 *
-	 * @param NewTopology The topology to set.
+	 * @param InTopology The topology to set.
+	 * @param InDuration Total duration of the media being played.
 	 * @return true on success, false otherwise.
+	 * @see Initialize, Shutdown
+	 */
+	bool SetTopology(const TComPtr<IMFTopology>& InTopology, FTimespan InDuration);
+
+	/**
+	 * Close the media session.
+	 *
+	 * @see Initialize
 	 */
-	bool SetTopology(const TComPtr<IMFTopology>& NewTopology);
+	void Shutdown();
 
 public:
 
 	//~ IMediaControls interface
 
+	virtual bool CanControl(EMediaControl Control) const override;
 	virtual FTimespan GetDuration() const override;
 	virtual float GetRate() const override;
 	virtual EMediaState GetState() const override;
-	virtual TRange<float> GetSupportedRates(EMediaPlaybackDirections Direction, bool Unthinned) const override;
+	virtual EMediaStatus GetStatus() const override;
+	virtual TRangeSet<float> GetSupportedRates(EMediaRateThinning Thinning) const override;
 	virtual FTimespan GetTime() const override;
 	virtual bool IsLooping() const override;
 	virtual bool Seek(const FTimespan& Time) override;
 	virtual bool SetLooping(bool Looping) override;
 	virtual bool SetRate(float Rate) override;
-	virtual bool SupportsRate(float Rate, bool Unthinned) const override;
-	virtual bool SupportsScrubbing() const override;
-	virtual bool SupportsSeeking() const override;
-
-public:
-
-	/**
-	 * Gets the session capabilities.
-	 *
-	 * @return Capabilities bit mask.
-	 */
-	DWORD GetCapabilities() const
-	{
-		return Capabilities;
-	}
-
-	/**
-	 * Sets the media state.
-	 *
-	 * @param NewState The media state to set.
-	 * @return true if the state will be changed, false otherwise.
-	 */
-	bool SetState(EMediaState NewState);
-
-public:
-
-	/** Gets an event delegate that is invoked when the last presentation segment finished playing. */
-	DECLARE_EVENT_OneParam(FWmfMediaSession, FOnSessionEvent, MediaEventType /*EventType*/)
-	FOnSessionEvent& OnSessionEvent()
-	{
-		return SessionEvent;
-	}
 
 public:
 
 	//~ IMFAsyncCallback interface
 
 	STDMETHODIMP_(ULONG) AddRef();
-	STDMETHODIMP GetParameters(unsigned long*, unsigned long*);
+	STDMETHODIMP GetParameters(unsigned long* Flags, unsigned long* Queue);
 	STDMETHODIMP Invoke(IMFAsyncResult* AsyncResult);
 	STDMETHODIMP QueryInterface(REFIID RefID, void** Object);
 	STDMETHODIMP_(ULONG) Release();
@@ -122,64 +114,87 @@ public:
 protected:
 
 	/**
-	 * Attempts to change the session state using the requested settings.
+	 * Commit the specified play rate.
+	 *
+	 * The caller holds the lock to the critical section.
+	 *
+	 * @param Rate The play rate to commit.
+	 * @param Thinning The thinning mode.
+	 * @see CommitTime, CommitTopology
+	 */
+	bool CommitRate(float Rate, EMediaRateThinning Thinning);
+
+	/**
+	 * Commit the specified play position.
+	 *
+	 * The caller holds the lock to the critical section.
 	 *
-	 * @return true if the state is being changed, false otherwise.
+	 * @param Time The play position to commit.
+	 * @see CommitRate, CommitTopology
 	 */
-	bool ChangeState();
+	bool CommitTime(FTimespan Time);
 
 	/**
-	 * Updates the CurrentPlayRate if it is different than the RequestedPlayRate.
+	 * Commit the given playback topology.
 	 *
-	 * @return Whether playback was stopped as a result of the update.
+	 * @param Topology The topology to set.
+	 * @see CommitRate, CommitTime
 	 */
-	bool UpdateRate();
+	bool CommitTopology(IMFTopology* Topology);
 
 	/**
-	 * Gets the playback position as known by the internal WMF clock.
+	 * Discard all pending state changes.
 	 *
-	 * @return The internal playback position.
+	 * @see UpdatePendingState
 	 */
-	FTimespan GetInternalPosition() const;
+	void DiscardPendingState();
 
 	/**
-	 * Updates the state machine after a state transition was completed.
+	 * Applies any pending state changes.
 	 *
-	 * @param CompletedState The state that was completed.
+	 * @see DiscardPendingState
 	 */
-	void UpdateState(EMediaState CompletedState);
+	void UpdatePendingState();
 
 private:
 
-	/** Whether the media source is buffering data. */
-	bool Buffering;
+	/** Hidden destructor (this class is reference counted). */
+	virtual ~FWmfMediaSession();
+
+private:
 
-	/** Caches whether the media supports scrubbing. */
+	/** Whether the media session supports scrubbing. */
 	bool CanScrub;
 
 	/** The session's capabilities. */
 	DWORD Capabilities;
 
-	/** Whether a new state change has been requested. */
-	bool ChangeRequested;
-
-	/** Critical section for locking access to this object. */
+	/** Synchronizes write access to session state. */
 	mutable FCriticalSection CriticalSection;
 
+	/** The duration of the media. */
+	FTimespan CurrentDuration;
+
 	/** The current playback rate. */
 	float CurrentRate;
 
 	/** The current playback state. */
 	EMediaState CurrentState;
 
-	/** The duration of the media. */
-	FTimespan Duration;
+	/** The full playback topology currently set on the media session. */
+	TComPtr<IMFTopology> CurrentTopology;
+
+	/** Media events to be forwarded to main thread. */
+	TQueue<EMediaEvent> DeferredEvents;
 
 	/** The media session that handles all playback. */
 	TComPtr<IMFMediaSession> MediaSession;
 
-	/** Whether playback should loop to the beginning. */
-	bool Looping;
+	/** The last play head position before playback was stopped. */
+	FTimespan LastTime;
+
+	/** Whether a state change is pending. */
+	bool PendingChange;
 
 	/** The media session's clock. */
 	TComPtr<IMFPresentationClock> PresentationClock;
@@ -193,28 +208,32 @@ private:
 	/** Holds a reference counter for this instance. */
 	int32 RefCount;
 
-	/** The requested playback rate. */
+	/** Deferred play rate change value. */
 	float RequestedRate;
 
-	/** The requested playback position. */
-	FTimespan RequestedTime;
+	/** Thinning of the deferred play rate. */
+	EMediaRateThinning RequestedThinning;
 
-	/** The requested playback state. */
-	EMediaState RequestedState;
+	/** Deferred playback topology to set. */
+	TComPtr<IMFTopology> RequestedTopology;
 
-	/** Whether a state change is currently in progress. */
-	bool StateChangePending;
+	/** Deferred play time change value. */
+	FTimespan RequestedTime;
+
+	/** Whether playback should loop to the beginning. */
+	bool ShouldLoop;
 
-	/** The playback topology to use. */
-	TComPtr<IMFTopology> Topology;
+	/** Current status flags. */
+	EMediaStatus Status;
 
-private:
+	/** The thinned play rates that the current media session supports. */
+	TRangeSet<float> ThinnedRates;
 
-	/** Holds an event delegate that is invoked when an event occurs in the session. */
-	FOnSessionEvent SessionEvent;
+	/** The unthinned play rates that the current media session supports. */
+	TRangeSet<float> UnthinnedRates;
 };
 
 
 #include "HideWindowsPlatformTypes.h"
 
-#endif
+#endif //WMFMEDIA_SUPPORTED_PLATFORM
diff --git a/Engine/Plugins/Media/WmfMedia/Source/WmfMedia/Private/Wmf/WmfMediaUtils.cpp b/Engine/Plugins/Media/WmfMedia/Source/WmfMedia/Private/Wmf/WmfMediaUtils.cpp
--- a/Engine/Plugins/Media/WmfMedia/Source/WmfMedia/Private/Wmf/WmfMediaUtils.cpp
+++ b/Engine/Plugins/Media/WmfMedia/Source/WmfMedia/Private/Wmf/WmfMediaUtils.cpp
@@ -1,29 +1,635 @@
 // Copyright 1998-2017 Epic Games, Inc. All Rights Reserved.
 
 #include "WmfMediaUtils.h"
-#include "HAL/PlatformProcess.h"
 
 #if WMFMEDIA_SUPPORTED_PLATFORM
 
+#include "HAL/FileManager.h"
+#include "HAL/PlatformMisc.h"
+#include "HAL/PlatformProcess.h"
+#include "Math/NumericLimits.h"
+#include "Misc/FileHelper.h"
+#include "Serialization/Archive.h"
+#include "Serialization/ArrayReader.h"
+
+#include "WmfMediaByteStream.h"
+
 #include "AllowWindowsPlatformTypes.h"
 
+#define WMFMEDIA_FRAMERATELUT_SIZE 9
+
+// MPEG-2 audio sub-types
+#define OTHER_FORMAT_MPEG2_AC3				0xe06d80e4 // MPEG-2 AC3
+#define OTHER_FORMAT_MPEG2_AUDIO			0xe06d802b // MPEG-2 Audio
+#define OTHER_FORMAT_MPEG2_DOLBY_AC3		0xe06d802c // Dolby AC3
+#define OTHER_FORMAT_MPEG2_DTS				0xe06d8033 // MPEG-2 DTS
+#define OTHER_FORMAT_MPEG2_LPCM_AUDIO		0xe06d8032 // DVD LPCM Audio
+#define OTHER_FORMAT_MPEG2_SDDS				0xe06d8034 // SDDS
+
+// MPEG-2 video sub-types
+#define OTHER_FORMAT_MPEG2_DVD_SUBPICTURE	0xe06d802d // DVD Sub-picture
+#define OTHER_FORMAT_MPEG2_VIDEO			0xe06d80e3 // MPEG-2 Video
+
+
+/* Local constants & helpers
+ *****************************************************************************/
 
 namespace WmfMedia
 {
+	// common media formats not defined by WMF
+	const GUID OtherFormatMpeg2_Base = { 0x00000000, 0xdb46, 0x11cf, { 0xb4, 0xd1, 0x00, 0x80, 0x05f, 0x6c, 0xbb, 0xea } };
+	const GUID OtherVideoFormat_LifeCam = { 0x3032344d, 0x0000, 0x0010, { 0x80, 0x00, 0x00, 0xAA, 0x00, 0x38, 0x9B, 0x71 } }; // Microsoft LifeCam UVC 1.0 video
 	const GUID OtherVideoFormat_QuickTime = { 0x61766331, 0x767a, 0x494d, { 0xb4, 0x78, 0xf2, 0x9d, 0x25, 0xdc, 0x90, 0x37 } }; // 1cva
 
 
+	/** Lookup table for frame rates that are handled specially in WMF. */
+	const struct FFrameRateLut
+	{
+		float FrameRate;
+		int32 Numerator;
+		int32 Denominator;
+		int32 DurationTicks;
+	}
+	FrameRateLut[WMFMEDIA_FRAMERATELUT_SIZE] = {
+		{ 59.95f, 60000, 1001, 166833 },
+		{ 29.97f, 30000, 1001, 333667 },
+		{ 23.976f, 24000, 1001, 417188 },
+		{ 60.0f, 60, 1, 166667 },
+		{ 30.0f, 30, 1, 333333 },
+		{ 50.0f, 50, 1, 200000 },
+		{ 25.0f, 25, 1, 400000 },
+		{ 24.0f, 24, 1, 416667 },
+		{ 0.0f, 0, 1, 0 }
+	};
+
+	/** List of supported major media types. */
+	GUID const* const SupportedMajorTypes[] = {
+		&MFMediaType_Audio,
+		&MFMediaType_Binary,
+		&MFMediaType_Image,
+		&MFMediaType_SAMI,
+		&MFMediaType_Video
+	};
+
+	/** List of supported audio channel counts (in order of preference). */
+	const DWORD SupportedAudioChannels[] = { 1, 2, 3, 4, 5, 6, 7, 8 };
+
+	/** List of supported audio types (in order of preference). */
+	const struct FAudioFormat
+	{
+		GUID const* const SubType;
+		UINT32 BitsPerSample;
+	}
+	SupportedAudioFormats[] = {
+		{ &MFAudioFormat_Float, 32u },
+		{ &MFAudioFormat_Float, 64u },
+		{ &MFAudioFormat_PCM, 32u },
+		{ &MFAudioFormat_PCM, 16u },
+		{ &MFAudioFormat_PCM, 8u }
+	};
+
+	/** List of supported video media types (in order of preference). */
+	GUID const* const SupportedVideoFormats[] =
+	{
+		// uncompressed
+//		&MFVideoFormat_A2R10G10B10,
+//		&MFVideoFormat_A16B16G16R16F,
+//		&MFVideoFormat_ARGB32
+		&MFVideoFormat_RGB24,
+		&MFVideoFormat_RGB32,
+//		&MFVideoFormat_RGB555,
+//		&MFVideoFormat_RGB565,
+//		&MFVideoFormat_RGB8,
+
+		// 8-bit YUV (packed)
+		&MFVideoFormat_AYUV,
+		&MFVideoFormat_UYVY,
+		&MFVideoFormat_YUY2,
+		&MFVideoFormat_YVYU,
+
+		// 8-bit YUV (planar)
+//		&MFVideoFormat_I420,
+//		&MFVideoFormat_IYUV,
+		&MFVideoFormat_NV12
+//		&MFVideoFormat_YV12,
+
+		// 10-bit YUV (packed)
+//		&MFVideoFormat_v210,
+//		&MFVideoFormat_v410,
+//		&MFVideoFormat_Y210,
+//		&MFVideoFormat_Y410,
+
+		// 10-bit YUV (planar)
+//		&MFVideoFormat_P010,
+//		&MFVideoFormat_P210,
+
+		// 16-bit YUV (packed)
+//		&MFVideoFormat_v216,
+//		&MFVideoFormat_Y216,
+//		&MFVideoFormat_Y416,
+
+		// 16-bit YUV (planar)
+//		&MFVideoFormat_P016,
+//		&MFVideoFormat_P216,
+	};
+
+
+	/** List of supported audio types. */
+	TArray<TComPtr<IMFMediaType>> SupportedAudioTypes;
+
+	/** List of supported binary types. */
+	TArray<TComPtr<IMFMediaType>> SupportedBinaryTypes;
+
+	/** List of supported SAMI types. */
+	TArray<TComPtr<IMFMediaType>> SupportedSamiTypes;
+
+	/** List of supported video types. */
+	TArray<TComPtr<IMFMediaType>> SupportedVideoTypes;
+
+
+	/** Initialize the lists of supported media types. */
+	void InitializeSupportedTypes()
+	{
+		static bool Initialized = false;
+
+		if (Initialized)
+		{
+			return;
+		}
+
+		// initialize audio types
+		for (const FAudioFormat& Format : SupportedAudioFormats)
+		{
+			for (UINT32 NumChannels : SupportedAudioChannels)
+			{
+				TComPtr<IMFMediaType> MediaType;
+
+				if (FAILED(::MFCreateMediaType(&MediaType)) ||
+					FAILED(MediaType->SetGUID(MF_MT_MAJOR_TYPE, MFMediaType_Audio)) ||
+					FAILED(MediaType->SetGUID(MF_MT_SUBTYPE, *Format.SubType)) ||
+					FAILED(MediaType->SetUINT32(MF_MT_ALL_SAMPLES_INDEPENDENT, TRUE)) ||
+					FAILED(MediaType->SetUINT32(MF_MT_AUDIO_BITS_PER_SAMPLE, Format.BitsPerSample)) ||
+					FAILED(MediaType->SetUINT32(MF_MT_AUDIO_NUM_CHANNELS, NumChannels)))
+				{
+					UE_LOG(LogWmfMedia, Error, TEXT("Failed to initialize supported audio types"));
+					return;
+				}
+
+				SupportedAudioTypes.Add(MediaType);
+			}
+		}
+
+		// initialize binary types
+		{
+			TComPtr<IMFMediaType> MediaType;
+
+			if (FAILED(::MFCreateMediaType(&MediaType)) ||
+				FAILED(MediaType->SetGUID(MF_MT_MAJOR_TYPE, MFMediaType_Binary)) ||
+				FAILED(MediaType->SetUINT32(MF_MT_ALL_SAMPLES_INDEPENDENT, TRUE)))
+			{
+				UE_LOG(LogWmfMedia, Error, TEXT("Failed to initialize supported binary types"));
+				return;
+			}
+
+			SupportedBinaryTypes.Add(MediaType);
+		}
+
+		// initialize SAMI types
+		{
+			TComPtr<IMFMediaType> MediaType;
+
+			if (FAILED(::MFCreateMediaType(&MediaType)) ||
+				FAILED(MediaType->SetGUID(MF_MT_MAJOR_TYPE, MFMediaType_SAMI)) ||
+				FAILED(MediaType->SetUINT32(MF_MT_ALL_SAMPLES_INDEPENDENT, TRUE)))
+			{
+				UE_LOG(LogWmfMedia, Error, TEXT("Failed to initialize supported SAMI types"));
+				return;
+			}
+
+			SupportedSamiTypes.Add(MediaType);
+		}
+
+		// initialize video types
+		for (GUID const* const Format : SupportedVideoFormats)
+		{
+			TComPtr<IMFMediaType> MediaType;
+
+			if (FAILED(::MFCreateMediaType(&MediaType)) ||
+				FAILED(MediaType->SetGUID(MF_MT_MAJOR_TYPE, MFMediaType_Video)) ||
+				FAILED(MediaType->SetGUID(MF_MT_SUBTYPE, *Format)) ||
+				FAILED(MediaType->SetUINT32(MF_MT_ALL_SAMPLES_INDEPENDENT, TRUE)))
+			{
+				UE_LOG(LogWmfMedia, Error, TEXT("Failed to initialize supported video types"));
+				return;
+			}
+
+			SupportedVideoTypes.Add(MediaType);
+		}
+	}
+}
+
+
+/* WmfMedia functions
+ *****************************************************************************/
+
+namespace WmfMedia
+{
+	FString CaptureDeviceRoleToString(ERole Role)
+	{
+		switch (Role)
+		{
+		case eCommunications:
+			return TEXT("Communications");
+
+		case eConsole:
+			return TEXT("Console");
+
+		case eMultimedia:
+			return TEXT("Multimedia");
+
+		default:
+			return TEXT("Unknown");
+		}
+	}
+
+
+	TComPtr<IMFMediaType> CreateOutputType(const GUID& MajorType, const GUID& SubType, bool AllowNonStandardCodecs)
+	{
+		TComPtr<IMFMediaType> OutputType;
+		{
+			HRESULT Result = ::MFCreateMediaType(&OutputType);
+
+			if (FAILED(Result))
+			{
+				UE_LOG(LogWmfMedia, Warning, TEXT("Failed to create %s output type (%s)"), *MajorTypeToString(MajorType), *ResultToString(Result));
+				return NULL;
+			}
+
+			Result = OutputType->SetUINT32(MF_MT_ALL_SAMPLES_INDEPENDENT, TRUE);
+
+			if (FAILED(Result))
+			{
+				UE_LOG(LogWmfMedia, Warning, TEXT("Failed to initialize %s output type (%s)"), *MajorTypeToString(MajorType), *ResultToString(Result));
+				return NULL;
+			}
+		}
+
+		if (MajorType == MFMediaType_Audio)
+		{
+			// filter unsupported audio formats
+			if (FMemory::Memcmp(&SubType.Data2, &MFMPEG4Format_Base.Data2, 12) == 0)
+			{
+				if (AllowNonStandardCodecs)
+				{
+					UE_LOG(LogWmfMedia, Verbose, TEXT("Allowing non-standard MP4 audio type %s (%s) \"%s\""), *SubTypeToString(SubType), *GuidToString(SubType), *FourccToString(SubType.Data1));
+				}
+				else
+				{
+					const bool DocumentedFormat =
+						(SubType.Data1 == WAVE_FORMAT_ADPCM) ||
+						(SubType.Data1 == WAVE_FORMAT_ALAW) ||
+						(SubType.Data1 == WAVE_FORMAT_MULAW) ||
+						(SubType.Data1 == WAVE_FORMAT_IMA_ADPCM) ||
+						(SubType.Data1 == MFAudioFormat_AAC.Data1) ||
+						(SubType.Data1 == MFAudioFormat_MP3.Data1) ||
+						(SubType.Data1 == MFAudioFormat_PCM.Data1);
+
+					const bool UndocumentedFormat =
+						(SubType.Data1 == WAVE_FORMAT_WMAUDIO2) ||
+						(SubType.Data1 == WAVE_FORMAT_WMAUDIO3) ||
+						(SubType.Data1 == WAVE_FORMAT_WMAUDIO_LOSSLESS);
+
+					if (!DocumentedFormat && !UndocumentedFormat)
+					{
+						UE_LOG(LogWmfMedia, Warning, TEXT("Skipping non-standard MP4 audio type %s (%s) \"%s\""), *SubTypeToString(SubType), *GuidToString(SubType), *FourccToString(SubType.Data1));
+						return NULL;
+					}
+				}
+			}
+			else if (FMemory::Memcmp(&SubType.Data2, &MFAudioFormat_Base.Data2, 12) != 0)
+			{
+				if (AllowNonStandardCodecs)
+				{
+					UE_LOG(LogWmfMedia, Verbose, TEXT("Allowing non-standard audio type %s (%s) \"%s\""), *SubTypeToString(SubType), *GuidToString(SubType), *FourccToString(SubType.Data1));
+				}
+				else
+				{
+					UE_LOG(LogWmfMedia, Warning, TEXT("Skipping non-standard audio type %s (%s) \"%s\""), *SubTypeToString(SubType), *GuidToString(SubType), *FourccToString(SubType.Data1));
+					return NULL;
+				}
+			}
+
+			// configure audio output (re-sampling fails for many media types, so we don't attempt it)
+			if (FAILED(OutputType->SetGUID(MF_MT_MAJOR_TYPE, MFMediaType_Audio)) ||
+				FAILED(OutputType->SetGUID(MF_MT_SUBTYPE, MFAudioFormat_PCM)) ||
+				FAILED(OutputType->SetUINT32(MF_MT_AUDIO_BITS_PER_SAMPLE, 16u)))
+			{
+				UE_LOG(LogWmfMedia, Warning, TEXT("Failed to initialize audio output type"));
+				return NULL;
+			}
+		}
+		else if (MajorType == MFMediaType_Binary)
+		{
+			const HRESULT Result = OutputType->SetGUID(MF_MT_MAJOR_TYPE, MFMediaType_Binary);
+
+			if (FAILED(Result))
+			{
+				UE_LOG(LogWmfMedia, Warning, TEXT("Failed to initialize binary output type (%s)"), *ResultToString(Result));
+				return NULL;
+			}
+		}
+		else if (MajorType == MFMediaType_SAMI)
+		{
+			// configure caption output
+			const HRESULT Result = OutputType->SetGUID(MF_MT_MAJOR_TYPE, MFMediaType_SAMI);
+
+			if (FAILED(Result))
+			{
+				UE_LOG(LogWmfMedia, Warning, TEXT("Failed to initialize caption output type (%s)"), *ResultToString(Result));
+				return NULL;
+			}
+		}
+		else if (MajorType == MFMediaType_Video)
+		{
+			// filter unsupported video types
+			if (FMemory::Memcmp(&SubType.Data2, &MFVideoFormat_Base.Data2, 12) != 0)
+			{
+				if (AllowNonStandardCodecs)
+				{
+					UE_LOG(LogWmfMedia, Verbose, TEXT("Allowing non-standard video type %s (%s) \"%s\""), *SubTypeToString(SubType), *GuidToString(SubType), *FourccToString(SubType.Data1));
+				}
+				else
+				{
+					UE_LOG(LogWmfMedia, Warning, TEXT("Skipping non-standard video type %s (%s) \"%s\""), *SubTypeToString(SubType), *GuidToString(SubType), *FourccToString(SubType.Data1));
+					return NULL;
+				}
+			}
+
+			if ((SubType == MFVideoFormat_HEVC) && !FWindowsPlatformMisc::VerifyWindowsVersion(10, 0))
+			{
+				UE_LOG(LogWmfMedia, Warning, TEXT("Your Windows version is %s"), *FPlatformMisc::GetOSVersion());
+
+				if ((SubType == MFVideoFormat_HEVC) && !FWindowsPlatformMisc::VerifyWindowsVersion(6, 2))
+				{
+					UE_LOG(LogWmfMedia, Warning, TEXT("HEVC video type requires Windows 10 or newer"));
+					return NULL;
+				}
+
+				UE_LOG(LogWmfMedia, Warning, TEXT("HEVC video type requires Windows 10 or newer (game must be manifested for Windows 10)"));
+			}
+
+			// configure video output
+			HRESULT Result = OutputType->SetGUID(MF_MT_MAJOR_TYPE, MFMediaType_Video);
+
+			if (FAILED(Result))
+			{
+				UE_LOG(LogWmfMedia, Warning, TEXT("Failed to set video output type (%s)"), *ResultToString(Result));
+				return NULL;
+			}
+
+			if (SubType == MFVideoFormat_HEVC)
+			{
+				Result = OutputType->SetGUID(MF_MT_SUBTYPE, MFVideoFormat_NV12);
+			}
+			else
+			{
+				const bool Uncompressed =
+					(SubType == MFVideoFormat_RGB555) ||
+					(SubType == MFVideoFormat_RGB565) ||
+					(SubType == MFVideoFormat_RGB24) ||
+					(SubType == MFVideoFormat_RGB32) ||
+					(SubType == MFVideoFormat_ARGB32);
+
+				Result = OutputType->SetGUID(MF_MT_SUBTYPE, Uncompressed ? MFVideoFormat_RGB32 : MFVideoFormat_YUY2);
+			}
+
+			if (FAILED(Result))
+			{
+				UE_LOG(LogWmfMedia, Warning, TEXT("Failed to set video output sub-type (%s)"), *ResultToString(Result));
+				return NULL;
+			}
+		}
+		else
+		{
+			return NULL; // unsupported input type
+		}
+
+		return OutputType;
+	}
+
+
+	void EnumerateCaptureDevices(GUID DeviceType, TArray<TComPtr<IMFActivate>>& OutDevices)
+	{
+		if ((DeviceType != MF_DEVSOURCE_ATTRIBUTE_SOURCE_TYPE_AUDCAP_GUID) &&
+			(DeviceType != MF_DEVSOURCE_ATTRIBUTE_SOURCE_TYPE_VIDCAP_GUID))
+		{
+			return; // unsupported device type
+		}
+
+		// create attribute store for search criteria
+		TComPtr<IMFAttributes> Attributes;
+		{
+			const HRESULT Result = ::MFCreateAttributes(&Attributes, 1);
+
+			if (FAILED(Result))
+			{
+				UE_LOG(LogWmfMedia, Error, TEXT("Failed to create capture device enumeration attributes: %s"), *WmfMedia::ResultToString(Result));
+				return;
+			}
+		}
+
+		// request capture devices
+		{
+			const HRESULT Result = Attributes->SetGUID(MF_DEVSOURCE_ATTRIBUTE_SOURCE_TYPE, DeviceType);
+
+			if (FAILED(Result))
+			{
+				UE_LOG(LogWmfMedia, Error, TEXT("Failed to set capture device enumeration type: %s"), *WmfMedia::ResultToString(Result));
+				return;
+			}
+		}
+
+		IMFActivate** Devices = nullptr;
+		UINT32 DeviceCount = 0;
+		{
+			const HRESULT Result = ::MFEnumDeviceSources(Attributes, &Devices, &DeviceCount);
+
+			if (DeviceCount == 0)
+			{
+				return; // no devices found
+			}
+		}
+
+		for (DWORD DeviceIndex = 0; DeviceIndex < DeviceCount; ++DeviceIndex)
+		{
+			TComPtr<IMFActivate> Device = Devices[DeviceIndex];
+			OutDevices.Add(Device);
+			Device->Release();
+		}
+
+		::CoTaskMemFree(Devices);
+	}
+
+
 	FString FourccToString(unsigned long Fourcc)
 	{
-		return FString::Printf(TEXT("%c%c%c%c"),
-			(Fourcc >> 24) & 0xff,
-			(Fourcc >> 16) & 0xff,
-			(Fourcc >> 8) & 0xff,
-			Fourcc & 0xff		
-		);
+		FString Result;
+
+		for (int32 CharIndex = 0; CharIndex < 4; ++CharIndex)
+		{
+			const unsigned char C = Fourcc & 0xff;
+			Result += FString::Printf(TChar<char>::IsPrint(C) ? TEXT("%c") : TEXT("[%d]"), C);
+			Fourcc >>= 8;
+		}
+
+		return Result;
 	}
 
 
+	bool FrameRateToRatio(float FrameRate, int32& OutNumerator, int32& OutDenominator)
+	{
+		if (FrameRate < 0.0f)
+		{
+			return false;
+		}
+
+		// use lookup table first to match WMF behavior
+		// see https://msdn.microsoft.com/en-us/library/windows/desktop/aa370467(v=vs.85).aspx
+
+		for (int32 LutIndex = 0; LutIndex < WMFMEDIA_FRAMERATELUT_SIZE; ++LutIndex)
+		{
+			const FFrameRateLut& Lut = FrameRateLut[LutIndex];
+
+			if (Lut.FrameRate == FrameRate)
+			{
+				OutNumerator = Lut.Numerator;
+				OutDenominator = Lut.Denominator;
+
+				return true;
+			}
+		}
+
+		// calculate a ratio (we could do better here, but this is fast)
+		const int32 NumeratorScale = 10000;
+
+		if (FrameRate > (MAX_int32 / NumeratorScale))
+		{
+			return false;
+		}
+
+		OutNumerator = (int32)(FrameRate * NumeratorScale);
+		OutDenominator = NumeratorScale;
+
+		return true;
+	}
+
+
+	bool GetCaptureDeviceInfo(IMFActivate& Device, FText& OutDisplayName, FString& OutInfo, bool& OutSoftwareDevice, FString& OutUrl)
+	{
+		GUID DeviceType;
+		{
+			if (FAILED(Device.GetGUID(MF_DEVSOURCE_ATTRIBUTE_SOURCE_TYPE, &DeviceType)))
+			{
+				return false; // failed to get device type
+			}
+
+			if ((DeviceType != MF_DEVSOURCE_ATTRIBUTE_SOURCE_TYPE_AUDCAP_GUID) &&
+				(DeviceType != MF_DEVSOURCE_ATTRIBUTE_SOURCE_TYPE_VIDCAP_GUID))
+			{
+				return false; // unsupported device type
+			}
+		}
+
+		const bool IsAudioDevice = (DeviceType == MF_DEVSOURCE_ATTRIBUTE_SOURCE_TYPE_AUDCAP_GUID);
+
+		// display name
+		PWSTR OutString = NULL;
+		UINT32 OutLength = 0;
+
+		if (SUCCEEDED(Device.GetAllocatedString(MF_DEVSOURCE_ATTRIBUTE_FRIENDLY_NAME, &OutString, &OutLength)))
+		{
+			OutDisplayName = FText::FromString(OutString);
+			::CoTaskMemFree(OutString);
+		}
+		else
+		{
+			OutDisplayName = NSLOCTEXT("WmfMedia", "UnknownCaptureDeviceName", "Unknown");
+		}
+
+		// debug information
+		if (IsAudioDevice)
+		{
+			UINT32 Role = 0;
+			if (SUCCEEDED(Device.GetUINT32(MF_DEVSOURCE_ATTRIBUTE_SOURCE_TYPE_AUDCAP_ROLE, &Role)))
+			{
+				OutInfo += TEXT("Role: ") + CaptureDeviceRoleToString((ERole)Role) + TEXT("\n");
+			}
+		}
+		else
+		{
+			UINT32 MaxBuffers = 0;
+			if (SUCCEEDED(Device.GetUINT32(MF_DEVSOURCE_ATTRIBUTE_SOURCE_TYPE_VIDCAP_MAX_BUFFERS, &MaxBuffers)))
+			{
+				OutInfo += FString::Printf(TEXT("Max Buffers: %i"), MaxBuffers);
+			}
+		}
+
+		// software device
+		UINT32 HwSource = 0;
+
+		if (SUCCEEDED(Device.GetUINT32(MF_DEVSOURCE_ATTRIBUTE_SOURCE_TYPE_VIDCAP_HW_SOURCE, &HwSource)))
+		{
+			OutSoftwareDevice = (HwSource == 0);
+		}
+		else
+		{
+			OutSoftwareDevice = false;
+		}
+
+		// symbolic link
+		const GUID SymbolicLinkAttribute = IsAudioDevice ? MF_DEVSOURCE_ATTRIBUTE_SOURCE_TYPE_AUDCAP_ENDPOINT_ID : MF_DEVSOURCE_ATTRIBUTE_SOURCE_TYPE_VIDCAP_SYMBOLIC_LINK;
+		const FString UrlScheme = IsAudioDevice ? TEXT("audcap://") : TEXT("vidcap://");
+
+		if (SUCCEEDED(Device.GetAllocatedString(SymbolicLinkAttribute, &OutString, &OutLength)))
+		{
+			OutUrl = UrlScheme + OutString;
+			::CoTaskMemFree(OutString);
+		}
+
+		return true;
+	}
+
+
+	const TArray<TComPtr<IMFMediaType>>& GetSupportedMediaTypes(const GUID& MajorType)
+	{
+		InitializeSupportedTypes();
+
+		if (MajorType == MFMediaType_Audio)
+		{
+			return SupportedAudioTypes;
+		}
+
+		if (MajorType == MFMediaType_Binary)
+		{
+			return SupportedBinaryTypes;
+		}
+
+		if (MajorType == MFMediaType_SAMI)
+		{
+			return SupportedSamiTypes;
+		}
+
+		if (MajorType == MFMediaType_Video)
+		{
+			return SupportedVideoTypes;
+		}
+
+		static TArray<TComPtr<IMFMediaType>> Empty;
+
+		return Empty;
+	}
+
+
 	FString GuidToString(const GUID& Guid)
 	{
 		return FString::Printf(TEXT("%08x-%04x-%04x-%02x%02x%02x%02x%02x%02x%02x%02x"),
@@ -42,6 +648,20 @@
 	}
 
 
+	bool IsSupportedMajorType(const GUID& MajorType)
+	{
+		for (GUID const* const Type : SupportedMajorTypes)
+		{
+			if (*Type == MajorType)
+			{
+				return true;
+			}
+		}
+
+		return false;
+	}
+
+
 	FString MajorTypeToString(const GUID& MajorType)
 	{
 		if (MajorType == MFMediaType_Default) return TEXT("Default");
@@ -60,3 +680,302 @@
 	}
 
 
+	FString MediaEventToString(MediaEventType Event)
+	{
+		switch (Event)
+		{
+		case MEUnknown: return TEXT("Unknown");
+		case MEError: return TEXT("Error");
+		case MEExtendedType: return TEXT("Extended Type");
+		case MENonFatalError: return TEXT("Non-fatal Error");
+		case MESessionUnknown: return TEXT("Session Unknown");
+		case MESessionTopologySet: return TEXT("Session Topology Set");
+		case MESessionTopologiesCleared: return TEXT("Session Topologies Cleared");
+		case MESessionStarted: return TEXT("Session Started");
+		case MESessionPaused: return TEXT("Session Paused");
+		case MESessionStopped: return TEXT("Session Stopped");
+		case MESessionClosed: return TEXT("Session Closed");
+		case MESessionEnded: return TEXT("Session Ended");
+		case MESessionRateChanged: return TEXT("Session Rate Changed");
+		case MESessionScrubSampleComplete: return TEXT("Session Scrub Sample Complete");
+		case MESessionCapabilitiesChanged: return TEXT("Session Capabilities Changed");
+		case MESessionTopologyStatus: return TEXT("Session Topology Status");
+		case MESessionNotifyPresentationTime: return TEXT("Session Notify Presentation Time");
+		case MENewPresentation: return TEXT("New Presentation");
+		case MELicenseAcquisitionStart: return TEXT("License Acquisition Start");
+		case MELicenseAcquisitionCompleted: return TEXT("License Acquisition Completed");
+		case MEIndividualizationStart: return TEXT("Individualization Start");
+		case MEIndividualizationCompleted: return TEXT("Individualization Completed");
+		case MEEnablerProgress: return TEXT("Enabler Progress");
+		case MEEnablerCompleted: return TEXT("Enabler Completed");
+		case MEPolicyError: return TEXT("Policy Error");
+		case MEPolicyReport: return TEXT("Policy Report");
+		case MEBufferingStarted: return TEXT("Buffering Started");
+		case MEBufferingStopped: return TEXT("Buffering Stopped");
+		case MEConnectStart: return TEXT("Connect Start");
+		case MEConnectEnd: return TEXT("Connect End");
+		case MEReconnectStart: return TEXT("Reconnect Start");
+		case MEReconnectEnd: return TEXT("Reconnect End");
+		case MERendererEvent: return TEXT("Renderer Event");
+		case MESessionStreamSinkFormatChanged: return TEXT("Session Stream Sink Format Changed");
+		case MESourceUnknown: return TEXT("Source Unknown");
+		case MESourceStarted: return TEXT("Source Started");
+		case MEStreamStarted: return TEXT("Stream Started");
+		case MESourceSeeked: return TEXT("Source Seeked");
+		case MEStreamSeeked: return TEXT("Stream Seeked");
+		case MENewStream: return TEXT("New Stream");
+		case MEUpdatedStream: return TEXT("Updated Stream");
+		case MESourceStopped: return TEXT("Source Stopped");
+		case MEStreamStopped: return TEXT("Stream Stopped");
+		case MESourcePaused: return TEXT("Source Paused");
+		case MEStreamPaused: return TEXT("Stream Paused");
+		case MEEndOfPresentation: return TEXT("End of Presentation");
+		case MEEndOfStream: return TEXT("End of Stream");
+		case MEMediaSample: return TEXT("Media Sample");
+		case MEStreamTick: return TEXT("Stream Tick");
+		case MEStreamThinMode: return TEXT("Stream Thin Mode");
+		case MEStreamFormatChanged: return TEXT("Stream Format Changed");
+		case MESourceRateChanged: return TEXT("Source Rate Changed");
+		case MEEndOfPresentationSegment: return TEXT("End of Presentation Segment");
+		case MESourceCharacteristicsChanged: return TEXT("Source Characteristics Changed");
+		case MESourceRateChangeRequested: return TEXT("Source Rate Change Requested");
+		case MESourceMetadataChanged: return TEXT("Source Metadata Changed");
+		case MESequencerSourceTopologyUpdated: return TEXT("Sequencer Source Topology Updated");
+		case MESinkUnknown: return TEXT("Sink Unknown");
+		case MEStreamSinkStarted: return TEXT("Stream Sink Started");
+		case MEStreamSinkStopped: return TEXT("Stream Sink Stopped");
+		case MEStreamSinkPaused: return TEXT("Strema Sink Paused");
+		case MEStreamSinkRateChanged: return TEXT("Stream Sink Rate Changed");
+		case MEStreamSinkRequestSample: return TEXT("Stream Sink Request Sample");
+		case MEStreamSinkMarker: return TEXT("Stream Sink Marker");
+		case MEStreamSinkPrerolled: return TEXT("Stream Sink Prerolled");
+		case MEStreamSinkScrubSampleComplete: return TEXT("Stream Sink Scrub Sample Complete");
+		case MEStreamSinkFormatChanged: return TEXT("Stream Sink Format Changed");
+		case MEStreamSinkDeviceChanged: return TEXT("Stream Sink Device Changed");
+		case MEQualityNotify: return TEXT("Quality Notify");
+		case MESinkInvalidated: return TEXT("Sink Invalidated");
+		case MEAudioSessionNameChanged: return TEXT("Audio Session Name Changed");
+		case MEAudioSessionVolumeChanged: return TEXT("Audio Session Volume Changed");
+		case MEAudioSessionDeviceRemoved: return TEXT("Audio Session Device Removed");
+		case MEAudioSessionServerShutdown: return TEXT("Audio Session Server Shutdown");
+		case MEAudioSessionGroupingParamChanged: return TEXT("Audio Session Grouping Param Changed");
+		case MEAudioSessionIconChanged: return TEXT("Audio Session Icion Changed");
+		case MEAudioSessionFormatChanged: return TEXT("Audio Session Format Changed");
+		case MEAudioSessionDisconnected: return TEXT("Audio Session Disconnected");
+		case MEAudioSessionExclusiveModeOverride: return TEXT("Audio Session Exclusive Mode Override");
+		case MECaptureAudioSessionVolumeChanged: return TEXT("Capture Audio Session Volume Changed");
+		case MECaptureAudioSessionDeviceRemoved: return TEXT("Capture Audio Session Device Removed");
+		case MECaptureAudioSessionFormatChanged: return TEXT("Capture Audio Session Format Changed");
+		case MECaptureAudioSessionDisconnected: return TEXT("Capture Audio Session Disconnected");
+		case MECaptureAudioSessionExclusiveModeOverride: return TEXT("Capture Audio Session Exclusive Mode Override");
+		case MECaptureAudioSessionServerShutdown: return TEXT("Capture Audio Session Server Shutdown");
+		case METrustUnknown: return TEXT("Trust Unknown");
+		case MEPolicyChanged: return TEXT("Policy Changed");
+		case MEContentProtectionMessage: return TEXT("Content Protection Message");
+		case MEPolicySet: return TEXT("Policy Set");
+		case MEWMDRMLicenseBackupCompleted: return TEXT("WM DRM License Backup Completed");
+		case MEWMDRMLicenseBackupProgress: return TEXT("WM DRM License Backup Progress");
+		case MEWMDRMLicenseRestoreCompleted: return TEXT("WM DRM License Restore Completed");
+		case MEWMDRMLicenseRestoreProgress: return TEXT("WM DRM License Restore Progress");
+		case MEWMDRMLicenseAcquisitionCompleted: return TEXT("WM DRM License Acquisition Completed");
+		case MEWMDRMIndividualizationCompleted: return TEXT("WM DRM Individualization Completed");
+		case MEWMDRMIndividualizationProgress: return TEXT("WM DRM Individualization Progress");
+		case MEWMDRMProximityCompleted: return TEXT("WM DRM Proximity Completed");
+		case MEWMDRMLicenseStoreCleaned: return TEXT("WM DRM License Store Cleaned");
+		case MEWMDRMRevocationDownloadCompleted: return TEXT("WM DRM Revocation Download Completed");
+		case METransformUnknown: return TEXT("Transform Unkonwn");
+		case METransformNeedInput: return TEXT("Transform Need Input");
+		case METransformHaveOutput: return TEXT("Transform Have Output");
+		case METransformDrainComplete: return TEXT("Transform Drain Complete");
+		case METransformMarker: return TEXT("Transform Marker");
+		case MEByteStreamCharacteristicsChanged: return TEXT("Byte Stream Characteristics Changed");
+		case MEVideoCaptureDeviceRemoved: return TEXT("Video Capture Device Removed");
+		case MEVideoCaptureDevicePreempted: return TEXT("Video Capture Device Preempted");
+		case MEStreamSinkFormatInvalidated: return TEXT("Stream Sink Format Invalidated");
+		case MEEncodingParameters: return TEXT("Encoding Paramters");
+		case MEContentProtectionMetadata: return TEXT("Content Protection Metadata");
+
+		default:
+			return FString::Printf(TEXT("Unknown event %i"), Event);
+		}
+	}
+
+
+	float RatioToFrameRate(int32 Numerator, int32 Denominator)
+	{
+		if (Denominator == 0)
+		{
+			return false;
+		}
+
+		// use lookup table first to match WMF behavior
+		// see https://msdn.microsoft.com/en-us/library/windows/desktop/aa370467(v=vs.85).aspx
+
+		for (int32 LutIndex = 0; LutIndex < WMFMEDIA_FRAMERATELUT_SIZE; ++LutIndex)
+		{
+			const FFrameRateLut& Lut = FrameRateLut[LutIndex];
+
+			if ((Lut.Numerator == Numerator) && (Lut.Denominator == Denominator))
+			{
+				return Lut.FrameRate;
+			}
+		}
+
+		return ((float)Numerator / (float)Denominator);
+	}
+
+
+	TComPtr<IMFMediaSource> ResolveMediaSource(TSharedPtr<FArchive, ESPMode::ThreadSafe> Archive, const FString& Url, bool Precache)
+	{
+		if (!Archive.IsValid())
+		{
+			// create capture device media source
+			if (Url.StartsWith(TEXT("audcap://")) || Url.StartsWith(TEXT("vidcap://")))
+			{
+				const TCHAR* EndpointOrSymlink = &Url[9];
+
+				TComPtr<IMFAttributes> Attributes;
+				{
+					HRESULT Result = ::MFCreateAttributes(&Attributes, 2);
+
+					if (FAILED(Result))
+					{
+						UE_LOG(LogWmfMedia, Error, TEXT("Failed to create capture device attributes (%s)"), *WmfMedia::ResultToString(Result));
+						return NULL;
+					}
+
+					const bool IsAudioDevice = Url.StartsWith(TEXT("audcap://"));
+
+					Result = Attributes->SetGUID(MF_DEVSOURCE_ATTRIBUTE_SOURCE_TYPE, IsAudioDevice
+						? MF_DEVSOURCE_ATTRIBUTE_SOURCE_TYPE_AUDCAP_GUID
+						: MF_DEVSOURCE_ATTRIBUTE_SOURCE_TYPE_VIDCAP_GUID);
+
+					if (FAILED(Result))
+					{
+						UE_LOG(LogWmfMedia, Error, TEXT("Failed to set capture device source type attribute (%s)"), *WmfMedia::ResultToString(Result));
+						return NULL;
+					}
+
+					if (IsAudioDevice)
+					{
+						Result = Attributes->SetString(MF_DEVSOURCE_ATTRIBUTE_SOURCE_TYPE_AUDCAP_ENDPOINT_ID, EndpointOrSymlink);
+					}
+					else
+					{
+						Result = Attributes->SetString(MF_DEVSOURCE_ATTRIBUTE_SOURCE_TYPE_VIDCAP_SYMBOLIC_LINK, EndpointOrSymlink);
+					}
+
+					if (FAILED(Result))
+					{
+						UE_LOG(LogWmfMedia, Error, TEXT("Failed to set capture device endpoint/symlink attribute (%s)"), *WmfMedia::ResultToString(Result));
+						return NULL;
+					}
+				}
+
+				TComPtr<IMFMediaSource> MediaSource;
+				{
+					const HRESULT Result = ::MFCreateDeviceSource(Attributes, &MediaSource);
+
+					if (FAILED(Result))
+					{
+						UE_LOG(LogWmfMedia, Error, TEXT("Failed to create capture device media source (%s)"), *WmfMedia::ResultToString(Result));
+						return NULL;
+					}
+				}
+
+				return MediaSource;
+			}
+
+			// load file media source
+			if (Url.StartsWith(TEXT("file://")))
+			{
+				const TCHAR* FilePath = &Url[7];
+
+				if (Precache)
+				{
+					FArrayReader* Reader = new FArrayReader;
+
+					if (FFileHelper::LoadFileToArray(*Reader, FilePath))
+					{
+						Archive = MakeShareable(Reader);
+					}
+					else
+					{
+						delete Reader;
+					}
+				}
+				else
+				{
+					Archive = MakeShareable(IFileManager::Get().CreateFileReader(FilePath));
+				}
+
+				if (!Archive.IsValid())
+				{
+					UE_LOG(LogWmfMedia, Error, TEXT("Failed to open or read media file %s"), FilePath);
+					return NULL;
+				}
+
+				if (Archive->TotalSize() == 0)
+				{
+					UE_LOG(LogWmfMedia, Error, TEXT("Cannot open media from empty file %s."), FilePath);
+					return NULL;
+				}
+			}
+		}
+
+		// create source resolver
+		TComPtr<IMFSourceResolver> SourceResolver;
+		{
+			const HRESULT Result = ::MFCreateSourceResolver(&SourceResolver);
+
+			if (FAILED(Result))
+			{
+				UE_LOG(LogWmfMedia, Error, TEXT("Failed to create media source resolver (%s)"), *WmfMedia::ResultToString(Result));
+				return NULL;
+			}
+		}
+
+		// resolve media source
+		TComPtr<IUnknown> SourceObject;
+		{
+			MF_OBJECT_TYPE ObjectType;
+
+			if (Archive.IsValid())
+			{
+				TComPtr<FWmfMediaByteStream> ByteStream = new FWmfMediaByteStream(Archive.ToSharedRef());
+				const HRESULT Result = SourceResolver->CreateObjectFromByteStream(ByteStream, *Url, MF_RESOLUTION_MEDIASOURCE, NULL, &ObjectType, &SourceObject);
+
+				if (FAILED(Result))
+				{
+					UE_LOG(LogWmfMedia, Error, TEXT("Failed to resolve byte stream %s (%s)"), *Url, *WmfMedia::ResultToString(Result));
+					return NULL;
+				}
+			}
+			else
+			{
+				const HRESULT Result = SourceResolver->CreateObjectFromURL(*Url, MF_RESOLUTION_MEDIASOURCE, NULL, &ObjectType, &SourceObject);
+
+				if (FAILED(Result))
+				{
+					UE_LOG(LogWmfMedia, Error, TEXT("Failed to resolve URL %s (%s)"), *Url, *WmfMedia::ResultToString(Result));
+					return NULL;
+				}
+			}
+		}
+
+		// get media source interface
+		TComPtr<IMFMediaSource> MediaSource;
+		{
+			const HRESULT Result = SourceObject->QueryInterface(IID_PPV_ARGS(&MediaSource));
+
+			if (FAILED(Result))
+			{
+				UE_LOG(LogWmfMedia, Error, TEXT("Failed to query media source interface: %s"), *WmfMedia::ResultToString(Result));
+				return NULL;
+			}
+		}
+
+		return MediaSource;
+	}
+
+
@@ -135,133 +1054,103 @@
 	}
 
 
-	FString SessionEventToString(MediaEventType Event)
-	{
-		switch (Event)
-		{
-		case MEUnknown: return TEXT("Unknown");
-		case MEError: return TEXT("Error");
-		case MEExtendedType: return TEXT("Extended Type");
-		case MENonFatalError: return TEXT("Non-fatal Error");
-		case MESessionUnknown: return TEXT("Session Unknown");
-		case MESessionTopologySet: return TEXT("Session Topology Set");
-		case MESessionTopologiesCleared: return TEXT("Session Topologies Cleared");
-		case MESessionStarted: return TEXT("Session Started");
-		case MESessionPaused: return TEXT("Session Paused");
-		case MESessionStopped: return TEXT("Session Stopped");
-		case MESessionClosed: return TEXT("Session Closed");
-		case MESessionEnded: return TEXT("Session Ended");
-		case MESessionRateChanged: return TEXT("Session Rate Changed");
-		case MESessionScrubSampleComplete: return TEXT("Session Scrub Sample Complete");
-		case MESessionCapabilitiesChanged: return TEXT("Session Capabilities Changed");
-		case MESessionTopologyStatus: return TEXT("Session Topology Status");
-		case MESessionNotifyPresentationTime: return TEXT("Session Notify Presentation Time");
-		case MENewPresentation: return TEXT("New Presentation");
-		case MELicenseAcquisitionStart: return TEXT("License Acquisition Start");
-		case MELicenseAcquisitionCompleted: return TEXT("License Acquisition Completed");
-		case MEIndividualizationStart: return TEXT("Individualization Start");
-		case MEIndividualizationCompleted: return TEXT("Individualization Completed");
-		case MEEnablerProgress: return TEXT("Enabler Progress");
-		case MEEnablerCompleted: return TEXT("Enabler Completed");
-		case MEPolicyError: return TEXT("Policy Error");
-		case MEPolicyReport: return TEXT("Policy Report");
-		case MEBufferingStarted: return TEXT("Buffering Started");
-		case MEBufferingStopped: return TEXT("Buffering Stopped");
-		case MEConnectStart: return TEXT("Connect Start");
-		case MEConnectEnd: return TEXT("Connect End");
-		case MEReconnectStart: return TEXT("Reconnect Start");
-		case MEReconnectEnd: return TEXT("Reconnect End");
-		case MERendererEvent: return TEXT("Renderer Event");
-		case MESessionStreamSinkFormatChanged: return TEXT("Session Stream Sink Format Changed");
-		case MESourceUnknown: return TEXT("Source Unknown");
-		case MESourceStarted: return TEXT("Source Started");
-		case MEStreamStarted: return TEXT("Stream Started");
-		case MESourceSeeked: return TEXT("Source Seeked");
-		case MEStreamSeeked: return TEXT("Stream Seeked");
-		case MENewStream: return TEXT("New Stream");
-		case MEUpdatedStream: return TEXT("Updated Stream");
-		case MESourceStopped: return TEXT("Source Stopped");
-		case MEStreamStopped: return TEXT("Stream Stopped");
-		case MESourcePaused: return TEXT("Source Paused");
-		case MEStreamPaused: return TEXT("Stream Paused");
-		case MEEndOfPresentation: return TEXT("End of Presentation");
-		case MEEndOfStream: return TEXT("End of Stream");
-		case MEMediaSample: return TEXT("Media Sample");
-		case MEStreamTick: return TEXT("Stream Tick");
-		case MEStreamThinMode: return TEXT("Stream Thin Mode");
-		case MEStreamFormatChanged: return TEXT("Stream Format Changed");
-		case MESourceRateChanged: return TEXT("Source Rate Changed");
-		case MEEndOfPresentationSegment: return TEXT("End of Presentation Segment");
-		case MESourceCharacteristicsChanged: return TEXT("Source Characteristics Changed");
-		case MESourceRateChangeRequested: return TEXT("Source Rate Change Requested");
-		case MESourceMetadataChanged: return TEXT("Source Metadata Changed");
-		case MESequencerSourceTopologyUpdated: return TEXT("Sequencer Source Topology Updated");
-		case MESinkUnknown: return TEXT("Sink Unknown");
-		case MEStreamSinkStarted: return TEXT("Stream Sink Started");
-		case MEStreamSinkStopped: return TEXT("Stream Sink Stopped");
-		case MEStreamSinkPaused: return TEXT("Strema Sink Paused");
-		case MEStreamSinkRateChanged: return TEXT("Stream Sink Rate Changed");
-		case MEStreamSinkRequestSample: return TEXT("Stream Sink Request Sample");
-		case MEStreamSinkMarker: return TEXT("Stream Sink Marker");
-		case MEStreamSinkPrerolled: return TEXT("Stream Sink Prerolled");
-		case MEStreamSinkScrubSampleComplete: return TEXT("Stream Sink Scrub Sample Complete");
-		case MEStreamSinkFormatChanged: return TEXT("Stream Sink Format Changed");
-		case MEStreamSinkDeviceChanged: return TEXT("Stream Sink Device Changed");
-		case MEQualityNotify: return TEXT("Quality Notify");
-		case MESinkInvalidated: return TEXT("Sink Invalidated");
-		case MEAudioSessionNameChanged: return TEXT("Audio Session Name Changed");
-		case MEAudioSessionVolumeChanged: return TEXT("Audio Session Volume Changed");
-		case MEAudioSessionDeviceRemoved: return TEXT("Audio Session Device Removed");
-		case MEAudioSessionServerShutdown: return TEXT("Audio Session Server Shutdown");
-		case MEAudioSessionGroupingParamChanged: return TEXT("Audio Session Grouping Param Changed");
-		case MEAudioSessionIconChanged: return TEXT("Audio Session Icion Changed");
-		case MEAudioSessionFormatChanged: return TEXT("Audio Session Format Changed");
-		case MEAudioSessionDisconnected: return TEXT("Audio Session Disconnected");
-		case MEAudioSessionExclusiveModeOverride: return TEXT("Audio Session Exclusive Mode Override");
-		case MECaptureAudioSessionVolumeChanged: return TEXT("Capture Audio Session Volume Changed");
-		case MECaptureAudioSessionDeviceRemoved: return TEXT("Capture Audio Session Device Removed");
-		case MECaptureAudioSessionFormatChanged: return TEXT("Capture Audio Session Format Changed");
-		case MECaptureAudioSessionDisconnected: return TEXT("Capture Audio Session Disconnected");
-		case MECaptureAudioSessionExclusiveModeOverride: return TEXT("Capture Audio Session Exclusive Mode Override");
-		case MECaptureAudioSessionServerShutdown: return TEXT("Capture Audio Session Server Shutdown");
-		case METrustUnknown: return TEXT("Trust Unknown");
-		case MEPolicyChanged: return TEXT("Policy Changed");
-		case MEContentProtectionMessage: return TEXT("Content Protection Message");
-		case MEPolicySet: return TEXT("Policy Set");
-		case MEWMDRMLicenseBackupCompleted: return TEXT("WM DRM License Backup Completed");
-		case MEWMDRMLicenseBackupProgress: return TEXT("WM DRM License Backup Progress");
-		case MEWMDRMLicenseRestoreCompleted: return TEXT("WM DRM License Restore Completed");
-		case MEWMDRMLicenseRestoreProgress: return TEXT("WM DRM License Restore Progress");
-		case MEWMDRMLicenseAcquisitionCompleted: return TEXT("WM DRM License Acquisition Completed");
-		case MEWMDRMIndividualizationCompleted: return TEXT("WM DRM Individualization Completed");
-		case MEWMDRMIndividualizationProgress: return TEXT("WM DRM Individualization Progress");
-		case MEWMDRMProximityCompleted: return TEXT("WM DRM Proximity Completed");
-		case MEWMDRMLicenseStoreCleaned: return TEXT("WM DRM License Store Cleaned");
-		case MEWMDRMRevocationDownloadCompleted: return TEXT("WM DRM Revocation Download Completed");
-		case METransformUnknown: return TEXT("Transform Unkonwn");
-		case METransformNeedInput: return TEXT("Transform Need Input");
-		case METransformHaveOutput: return TEXT("Transform Have Output");
-		case METransformDrainComplete: return TEXT("Transform Drain Complete");
-		case METransformMarker: return TEXT("Transform Marker");
-		case MEByteStreamCharacteristicsChanged: return TEXT("Byte Stream Characteristics Changed");
-		case MEVideoCaptureDeviceRemoved: return TEXT("Video Capture Device Removed");
-		case MEVideoCaptureDevicePreempted: return TEXT("Video Capture Device Preempted");
-		case MEStreamSinkFormatInvalidated: return TEXT("Stream Sink Format Invalidated");
-		case MEEncodingParameters: return TEXT("Encoding Paramters");
-		case MEContentProtectionMetadata: return TEXT("Content Protection Metadata");
-		default:
-			return FString::Printf(TEXT("Unknown event %i"), Event);
-		}
-	}
-
-
 	FString SubTypeToString(const GUID& SubType)
 	{
+		if (SubType == GUID_NULL) return TEXT("Null");
+
+		// image formats
+		if (SubType == MFImageFormat_JPEG) return TEXT("Jpeg");
+		if (SubType == MFImageFormat_RGB32) return TEXT("RGB32");
+
+		// stream formats
+		if (SubType == MFStreamFormat_MPEG2Transport) return TEXT("MPEG-2 Transport");
+		if (SubType == MFStreamFormat_MPEG2Program) return TEXT("MPEG-2 Program");
+
+		// video formats
+		if (SubType == MFVideoFormat_RGB32) return TEXT("RGB32");
+		if (SubType == MFVideoFormat_ARGB32) return TEXT("ARGB32");
+		if (SubType == MFVideoFormat_RGB24) return TEXT("RGB24");
+		if (SubType == MFVideoFormat_RGB555) return TEXT("RGB525");
+		if (SubType == MFVideoFormat_RGB565) return TEXT("RGB565");
+		if (SubType == MFVideoFormat_RGB8) return TEXT("RGB8");
+		if (SubType == MFVideoFormat_AI44) return TEXT("AI44");
+		if (SubType == MFVideoFormat_AYUV) return TEXT("AYUV");
+		if (SubType == MFVideoFormat_YUY2) return TEXT("YUY2");
+		if (SubType == MFVideoFormat_YVYU) return TEXT("YVYU");
+		if (SubType == MFVideoFormat_YVU9) return TEXT("YVU9");
+		if (SubType == MFVideoFormat_UYVY) return TEXT("UYVY");
+		if (SubType == MFVideoFormat_NV11) return TEXT("NV11");
+		if (SubType == MFVideoFormat_NV12) return TEXT("NV12");
+		if (SubType == MFVideoFormat_YV12) return TEXT("YV12");
+		if (SubType == MFVideoFormat_I420) return TEXT("I420");
+		if (SubType == MFVideoFormat_IYUV) return TEXT("IYUV");
+		if (SubType == MFVideoFormat_Y210) return TEXT("Y210");
+		if (SubType == MFVideoFormat_Y216) return TEXT("Y216");
+		if (SubType == MFVideoFormat_Y410) return TEXT("Y410");
+		if (SubType == MFVideoFormat_Y416) return TEXT("Y416");
+		if (SubType == MFVideoFormat_Y41P) return TEXT("Y41P");
+		if (SubType == MFVideoFormat_Y41T) return TEXT("Y41T");
+		if (SubType == MFVideoFormat_Y42T) return TEXT("Y42T");
+		if (SubType == MFVideoFormat_P210) return TEXT("P210");
+		if (SubType == MFVideoFormat_P216) return TEXT("P216");
+		if (SubType == MFVideoFormat_P010) return TEXT("P010");
+		if (SubType == MFVideoFormat_P016) return TEXT("P016");
+		if (SubType == MFVideoFormat_v210) return TEXT("v210");
+		if (SubType == MFVideoFormat_v216) return TEXT("v216");
+		if (SubType == MFVideoFormat_v410) return TEXT("v410");
+		if (SubType == MFVideoFormat_MP43) return TEXT("MP43");
+		if (SubType == MFVideoFormat_MP4S) return TEXT("MP4S");
+		if (SubType == MFVideoFormat_M4S2) return TEXT("M4S2");
+		if (SubType == MFVideoFormat_MP4V) return TEXT("MP4V");
+		if (SubType == MFVideoFormat_WMV1) return TEXT("WMV1");
+		if (SubType == MFVideoFormat_WMV2) return TEXT("WMV2");
+		if (SubType == MFVideoFormat_WMV3) return TEXT("WMV3");
+		if (SubType == MFVideoFormat_WVC1) return TEXT("WVC1");
+		if (SubType == MFVideoFormat_MSS1) return TEXT("MSS1");
+		if (SubType == MFVideoFormat_MSS2) return TEXT("MSS2");
+		if (SubType == MFVideoFormat_MPG1) return TEXT("MPG1");
+		if (SubType == MFVideoFormat_DVSL) return TEXT("DVSL");
+		if (SubType == MFVideoFormat_DVSD) return TEXT("DVSD");
+		if (SubType == MFVideoFormat_DVHD) return TEXT("DVHD");
+		if (SubType == MFVideoFormat_DV25) return TEXT("DV25");
+		if (SubType == MFVideoFormat_DV50) return TEXT("DV50");
+		if (SubType == MFVideoFormat_DVH1) return TEXT("DVH1");
+		if (SubType == MFVideoFormat_DVC) return TEXT("DVC");
+		if (SubType == MFVideoFormat_H264) return TEXT("H264");
+		if (SubType == MFVideoFormat_MJPG) return TEXT("MJPG");
+		if (SubType == MFVideoFormat_420O) return TEXT("420O");
+		if (SubType == MFVideoFormat_HEVC) return TEXT("HEVC");
+		if (SubType == MFVideoFormat_HEVC_ES) return TEXT("HEVC ES");
+
+#if (WINVER >= _WIN32_WINNT_WIN8)
+		if (SubType == MFVideoFormat_H263) return TEXT("H263");
+#endif
+
+		if (SubType == MFVideoFormat_H264_ES) return TEXT("H264 ES");
+		if (SubType == MFVideoFormat_MPEG2) return TEXT("MPEG-2");
+
+		// common non-Windows formats
+		if (SubType == OtherVideoFormat_LifeCam) return TEXT("LifeCam");
+		if (SubType == OtherVideoFormat_QuickTime) return TEXT("QuickTime");
+
+		if (FMemory::Memcmp(&SubType.Data2, &OtherFormatMpeg2_Base.Data2, 12) == 0)
+		{
+			if (SubType.Data1 == OTHER_FORMAT_MPEG2_AC3) return TEXT("MPEG-2 AC3");
+			if (SubType.Data1 == OTHER_FORMAT_MPEG2_AUDIO) return TEXT("MPEG-2 Audio");
+			if (SubType.Data1 == OTHER_FORMAT_MPEG2_DOLBY_AC3) return TEXT("Dolby AC-3");
+			if (SubType.Data1 == OTHER_FORMAT_MPEG2_DTS) return TEXT("DTS");
+			if (SubType.Data1 == OTHER_FORMAT_MPEG2_LPCM_AUDIO) return TEXT("DVD LPCM");
+			if (SubType.Data1 == OTHER_FORMAT_MPEG2_SDDS) return TEXT("SDDS");
+
+			if (SubType.Data1 == OTHER_FORMAT_MPEG2_DVD_SUBPICTURE) return TEXT("DVD Subpicture");
+			if (SubType.Data1 == OTHER_FORMAT_MPEG2_VIDEO) return TEXT("MPEG-2 Video");
+		}
+
 		// audio formats
 		if ((FMemory::Memcmp(&SubType.Data2, &MFAudioFormat_Base.Data2, 12) == 0) ||
 			(FMemory::Memcmp(&SubType.Data2, &MFMPEG4Format_Base.Data2, 12) == 0))
 		{
-			if (SubType.Data1 == WAVE_FORMAT_UNKNOWN) return TEXT("Unknown Audio Format");\
+			if (SubType.Data1 == WAVE_FORMAT_UNKNOWN) return TEXT("Unknown Audio Format");
 			if (SubType.Data1 == WAVE_FORMAT_PCM) return TEXT("PCM");
 			if (SubType.Data1 == WAVE_FORMAT_ADPCM) return TEXT("ADPCM");
 			if (SubType.Data1 == WAVE_FORMAT_IEEE_FLOAT) return TEXT("IEEE Float");
@@ -527,80 +1416,7 @@
 			if (SubType.Data1 == WAVE_FORMAT_FLAC) return TEXT("flac.sourceforge.net");
 		}
 
-		// image formats
-		if (SubType == MFImageFormat_JPEG) return TEXT("Jpeg");
-		if (SubType == MFImageFormat_RGB32) return TEXT("RGB32");
-
-		// stream formats
-		if (SubType == MFStreamFormat_MPEG2Transport) return TEXT("MPEG-2 Transport");
-		if (SubType == MFStreamFormat_MPEG2Program) return TEXT("MPEG-2 Program");
-
-		// video formats
-		if (SubType == MFVideoFormat_RGB32) return TEXT("RGB32");
-		if (SubType == MFVideoFormat_ARGB32) return TEXT("ARGB32");
-		if (SubType == MFVideoFormat_RGB24) return TEXT("RGB24");
-		if (SubType == MFVideoFormat_RGB555) return TEXT("RGB525");
-		if (SubType == MFVideoFormat_RGB565) return TEXT("RGB565");
-		if (SubType == MFVideoFormat_RGB8) return TEXT("RGB8");
-		if (SubType == MFVideoFormat_AI44) return TEXT("AI44");
-		if (SubType == MFVideoFormat_AYUV) return TEXT("AYUV");
-		if (SubType == MFVideoFormat_YUY2) return TEXT("YUY2");
-		if (SubType == MFVideoFormat_YVYU) return TEXT("YVYU");
-		if (SubType == MFVideoFormat_YVU9) return TEXT("YVU9");
-		if (SubType == MFVideoFormat_UYVY) return TEXT("UYVY");
-		if (SubType == MFVideoFormat_NV11) return TEXT("NV11");
-		if (SubType == MFVideoFormat_NV12) return TEXT("NV12");
-		if (SubType == MFVideoFormat_YV12) return TEXT("YV12");
-		if (SubType == MFVideoFormat_I420) return TEXT("I420");
-		if (SubType == MFVideoFormat_IYUV) return TEXT("IYUV");
-		if (SubType == MFVideoFormat_Y210) return TEXT("Y210");
-		if (SubType == MFVideoFormat_Y216) return TEXT("Y216");
-		if (SubType == MFVideoFormat_Y410) return TEXT("Y410");
-		if (SubType == MFVideoFormat_Y416) return TEXT("Y416");
-		if (SubType == MFVideoFormat_Y41P) return TEXT("Y41P");
-		if (SubType == MFVideoFormat_Y41T) return TEXT("Y41T");
-		if (SubType == MFVideoFormat_Y42T) return TEXT("Y42T");
-		if (SubType == MFVideoFormat_P210) return TEXT("P210");
-		if (SubType == MFVideoFormat_P216) return TEXT("P216");
-		if (SubType == MFVideoFormat_P010) return TEXT("P010");
-		if (SubType == MFVideoFormat_P016) return TEXT("P016");
-		if (SubType == MFVideoFormat_v210) return TEXT("v210");
-		if (SubType == MFVideoFormat_v216) return TEXT("v216");
-		if (SubType == MFVideoFormat_v410) return TEXT("v410");
-		if (SubType == MFVideoFormat_MP43) return TEXT("MP43");
-		if (SubType == MFVideoFormat_MP4S) return TEXT("MP4S");
-		if (SubType == MFVideoFormat_M4S2) return TEXT("M4S2");
-		if (SubType == MFVideoFormat_MP4V) return TEXT("MP4V");
-		if (SubType == MFVideoFormat_WMV1) return TEXT("WMV1");
-		if (SubType == MFVideoFormat_WMV2) return TEXT("WMV2");
-		if (SubType == MFVideoFormat_WMV3) return TEXT("WMV3");
-		if (SubType == MFVideoFormat_WVC1) return TEXT("WVC1");
-		if (SubType == MFVideoFormat_MSS1) return TEXT("MSS1");
-		if (SubType == MFVideoFormat_MSS2) return TEXT("MSS2");
-		if (SubType == MFVideoFormat_MPG1) return TEXT("MPG1");
-		if (SubType == MFVideoFormat_DVSL) return TEXT("DVSL");
-		if (SubType == MFVideoFormat_DVSD) return TEXT("DVSD");
-		if (SubType == MFVideoFormat_DVHD) return TEXT("DVHD");
-		if (SubType == MFVideoFormat_DV25) return TEXT("DV25");
-		if (SubType == MFVideoFormat_DV50) return TEXT("DV50");
-		if (SubType == MFVideoFormat_DVH1) return TEXT("DVH1");
-		if (SubType == MFVideoFormat_DVC) return TEXT("DVC");
-		if (SubType == MFVideoFormat_H264) return TEXT("H264");
-		if (SubType == MFVideoFormat_MJPG) return TEXT("MJPG");
-		if (SubType == MFVideoFormat_420O) return TEXT("420O");
-		if (SubType == MFVideoFormat_HEVC) return TEXT("HEVC");
-		if (SubType == MFVideoFormat_HEVC_ES) return TEXT("HEVC ES");
-
-#if (WINVER >= _WIN32_WINNT_WIN8)
-		if (SubType == MFVideoFormat_H263) return TEXT("H263");
-#endif
-
-		if (SubType == MFVideoFormat_H264_ES) return TEXT("H264 ES");
-		if (SubType == MFVideoFormat_MPEG2) return TEXT("MPEG-2");
-
-		// common non-Windows formats
-		if (SubType == OtherVideoFormat_QuickTime) return TEXT("QuickTime");
-
+		// unknown type
 		return FString::Printf(TEXT("%s (%s)"), *GuidToString(SubType), *FourccToString(SubType.Data1));
 	}
 }
