Copying //UE4/Dev-Sequencer to //UE4/Main (Source: //UE4/Dev-Sequencer @ 3631098)

diff --git a/Engine/Plugins/Media/WmfMedia/Source/WmfMedia/Private/Player/WmfMediaTracks.cpp b/Engine/Plugins/Media/WmfMedia/Source/WmfMedia/Private/Player/WmfMediaTracks.cpp
--- a/Engine/Plugins/Media/WmfMedia/Source/WmfMedia/Private/Player/WmfMediaTracks.cpp
+++ b/Engine/Plugins/Media/WmfMedia/Source/WmfMedia/Private/Player/WmfMediaTracks.cpp
@@ -7,6 +7,7 @@
 
 #include "IHeadMountedDisplayModule.h"
 #include "IMediaOptions.h"
+#include "MediaHelpers.h"
 #include "Misc/ScopeLock.h"
 #include "UObject/Class.h"
 
@@ -34,10 +35,10 @@ FWmfMediaTracks::FWmfMediaTracks()
 	: AudioSamplePool(new FWmfMediaAudioSamplePool)
 	, MediaSourceChanged(false)
 	, SelectedAudioTrack(INDEX_NONE)
-	, SelectedBinaryTrack(INDEX_NONE)
 	, SelectedCaptionTrack(INDEX_NONE)
+	, SelectedMetadataTrack(INDEX_NONE)
 	, SelectedVideoTrack(INDEX_NONE)
-	, TopologyChanged(false)
+	, SelectionChanged(false)
 	, VideoSamplePool(new FWmfMediaTextureSamplePool)
 { }
 
@@ -66,14 +67,14 @@ void FWmfMediaTracks::AppendStats(FString &OutStats) const
 	
 	if (AudioTracks.Num() == 0)
 	{
-		OutStats += TEXT("    none\n");
+		OutStats += TEXT("\tnone\n");
 	}
 	else
 	{
 		for (const FTrack& Track : AudioTracks)
 		{
-			OutStats += FString::Printf(TEXT("    %s\n"), *Track.DisplayName.ToString());
-			OutStats += TEXT("        Not implemented yet");
+			OutStats += FString::Printf(TEXT("\t%s\n"), *Track.DisplayName.ToString());
+			OutStats += TEXT("\t\tNot implemented yet");
 		}
 	}
 
@@ -82,14 +83,14 @@ void FWmfMediaTracks::AppendStats(FString &OutStats) const
 
 	if (VideoTracks.Num() == 0)
 	{
-		OutStats += TEXT("    none\n");
+		OutStats += TEXT("\tnone\n");
 	}
 	else
 	{
 		for (const FTrack& Track : VideoTracks)
 		{
-			OutStats += FString::Printf(TEXT("    %s\n"), *Track.DisplayName.ToString());
-			OutStats += TEXT("        Not implemented yet");
+			OutStats += FString::Printf(TEXT("\t%s\n"), *Track.DisplayName.ToString());
+			OutStats += TEXT("\t\tNot implemented yet");
 		}
 	}
 }
@@ -100,7 +101,7 @@ void FWmfMediaTracks::ClearFlags()
 	FScopeLock Lock(&CriticalSection);
 
 	MediaSourceChanged = false;
-	TopologyChanged = false;
+	SelectionChanged = false;
 }
 
 
@@ -113,8 +114,8 @@ TComPtr<IMFTopology> FWmfMediaTracks::CreateTopology()
 	}
 
 	if ((SelectedAudioTrack == INDEX_NONE) &&
-		(SelectedBinaryTrack == INDEX_NONE) &&
 		(SelectedCaptionTrack == INDEX_NONE) &&
+		(SelectedMetadataTrack == INDEX_NONE) &&
 		(SelectedVideoTrack == INDEX_NONE))
 	{
 		return NULL; // no tracks selected
@@ -127,7 +128,7 @@ TComPtr<IMFTopology> FWmfMediaTracks::CreateTopology()
 
 		if (FAILED(Result))
 		{
-			UE_LOG(LogWmfMedia, Error, TEXT("Failed to create playback topology: %s"), *WmfMedia::ResultToString(Result));
+			UE_LOG(LogWmfMedia, Verbose, TEXT("Tracks %p: Failed to create playback topology: %s"), this, *WmfMedia::ResultToString(Result));
 			return NULL;
 		}
 	}
@@ -140,14 +141,14 @@ TComPtr<IMFTopology> FWmfMediaTracks::CreateTopology()
 		TracksAdded |= AddTrackToTopology(AudioTracks[SelectedAudioTrack], *Topology);
 	}
 
-	if (MetadataTracks.IsValidIndex(SelectedBinaryTrack))
+	if (CaptionTracks.IsValidIndex(SelectedCaptionTrack))
 	{
-		TracksAdded |= AddTrackToTopology(MetadataTracks[SelectedBinaryTrack], *Topology);
+		TracksAdded |= AddTrackToTopology(CaptionTracks[SelectedCaptionTrack], *Topology);
 	}
 
-	if (CaptionTracks.IsValidIndex(SelectedCaptionTrack))
+	if (MetadataTracks.IsValidIndex(SelectedMetadataTrack))
 	{
-		TracksAdded |= AddTrackToTopology(CaptionTracks[SelectedCaptionTrack], *Topology);
+		TracksAdded |= AddTrackToTopology(MetadataTracks[SelectedMetadataTrack], *Topology);
 	}
 
 	if (VideoTracks.IsValidIndex(SelectedVideoTrack))
@@ -180,12 +181,12 @@ FTimespan FWmfMediaTracks::GetDuration() const
 }
 
 
-void FWmfMediaTracks::GetFlags(bool& OutMediaSourceChanged, bool& OutTopologyChanged) const
+void FWmfMediaTracks::GetFlags(bool& OutMediaSourceChanged, bool& OutSelectionChanged) const
 {
 	FScopeLock Lock(&CriticalSection);
 
 	OutMediaSourceChanged = MediaSourceChanged;
-	OutTopologyChanged = TopologyChanged;
+	OutSelectionChanged = SelectionChanged;
 }
 
 
@@ -193,64 +194,76 @@ void FWmfMediaTracks::Initialize(IMFMediaSource* InMediaSource)
 {
 	Shutdown();
 
+	UE_LOG(LogWmfMedia, Verbose, TEXT("Tracks: %p: Initializing tracks"), this);
+
 	FScopeLock Lock(&CriticalSection);
 
 	MediaSourceChanged = true;
+	SelectionChanged = true;
 
-	if (InMediaSource != NULL)
+	if (InMediaSource == NULL)
 	{
-		// create presentation descriptor
-		TComPtr<IMFPresentationDescriptor> NewPresentationDescriptor;
-		{
-			const HRESULT Result = InMediaSource->CreatePresentationDescriptor(&NewPresentationDescriptor);
+		return;
+	}
 
-			if (FAILED(Result))
-			{
-				UE_LOG(LogWmfMedia, Error, TEXT("Failed to create presentation descriptor: %s"), *WmfMedia::ResultToString(Result));
-				return;
-			}
-		}
+	// create presentation descriptor
+	TComPtr<IMFPresentationDescriptor> NewPresentationDescriptor;
+	{
+		const HRESULT Result = InMediaSource->CreatePresentationDescriptor(&NewPresentationDescriptor);
 
-		// get number of streams
-		DWORD StreamCount = 0;
+		if (FAILED(Result))
 		{
-			const HRESULT Result = NewPresentationDescriptor->GetStreamDescriptorCount(&StreamCount);
-
-			if (FAILED(Result))
-			{
-				UE_LOG(LogWmfMedia, Error, TEXT("Failed to get stream count (%s)"), *WmfMedia::ResultToString(Result));
-				return;
-			}
-
-			UE_LOG(LogWmfMedia, Verbose, TEXT("Tracks %llx: Found %i streams"), this, StreamCount);
+			UE_LOG(LogWmfMedia, Verbose, TEXT("Tracks %p: Failed to create presentation descriptor: %s"), *WmfMedia::ResultToString(Result));
+			return;
 		}
+	}
 
-		// initialization successful
-		MediaSource = InMediaSource;
-		PresentationDescriptor = NewPresentationDescriptor;
+	// get number of streams
+	DWORD StreamCount = 0;
+	{
+		const HRESULT Result = NewPresentationDescriptor->GetStreamDescriptorCount(&StreamCount);
 
-		// add streams (Media Foundation reports them in reverse order)
-		for (int32 StreamIndex = StreamCount - 1; StreamIndex >= 0; --StreamIndex)
+		if (FAILED(Result))
 		{
-			AddStreamToTracks(StreamIndex, Info);
-			Info += TEXT("\n");
+			UE_LOG(LogWmfMedia, Verbose, TEXT("Tracks %p: Failed to get stream count: %s"), this, *WmfMedia::ResultToString(Result));
+			return;
 		}
+
+		UE_LOG(LogWmfMedia, Verbose, TEXT("Tracks %p: Found %i streams"), this, StreamCount);
 	}
 
-	TopologyChanged = true;
+	// initialization successful
+	MediaSource = InMediaSource;
+	PresentationDescriptor = NewPresentationDescriptor;
+
+	// add streams (Media Foundation reports them in reverse order)
+	bool AllStreamsAdded = true;
+
+	for (int32 StreamIndex = StreamCount - 1; StreamIndex >= 0; --StreamIndex)
+	{
+		AllStreamsAdded &= AddStreamToTracks(StreamIndex, Info);
+		Info += TEXT("\n");
+	}
+
+	if (!AllStreamsAdded)
+	{
+		UE_LOG(LogWmfMedia, Verbose, TEXT("Tracks %p: Not all available streams were added to the track collection"), this);
+	}
 }
 
 
 void FWmfMediaTracks::Shutdown()
 {
+	UE_LOG(LogWmfMedia, Verbose, TEXT("Tracks: %p: Shutting down tracks"), this);
+
 	FScopeLock Lock(&CriticalSection);
 
 	AudioSamplePool->Reset();
 	VideoSamplePool->Reset();
 
 	SelectedAudioTrack = INDEX_NONE;
-	SelectedBinaryTrack = INDEX_NONE;
 	SelectedCaptionTrack = INDEX_NONE;
+	SelectedMetadataTrack = INDEX_NONE;
 	SelectedVideoTrack = INDEX_NONE;
 
 	AudioTracks.Empty();
@@ -269,7 +282,7 @@ void FWmfMediaTracks::Shutdown()
 	PresentationDescriptor.Reset();
 
 	MediaSourceChanged = false;
-	TopologyChanged = false;
+	SelectionChanged = false;
 }
 
 
@@ -487,12 +500,12 @@ int32 FWmfMediaTracks::GetSelectedTrack(EMediaTrackType TrackType) const
 	case EMediaTrackType::Audio:
 		return SelectedAudioTrack;
 
-	case EMediaTrackType::Metadata:
-		return SelectedBinaryTrack;
-
 	case EMediaTrackType::Caption:
 		return SelectedCaptionTrack;
 
+	case EMediaTrackType::Metadata:
+		return SelectedMetadataTrack;
+
 	case EMediaTrackType::Video:
 		return SelectedVideoTrack;
 
@@ -666,201 +679,155 @@ bool FWmfMediaTracks::SelectTrack(EMediaTrackType TrackType, int32 TrackIndex)
 		return false; // not initialized
 	}
 
+	UE_LOG(LogWmfMedia, Verbose, TEXT("Tracks %p: Selecting %s track %i"), this, *MediaUtils::TrackTypeToString(TrackType), TrackIndex);
+
 	FScopeLock Lock(&CriticalSection);
 
+	int32* SelectedTrack = nullptr;
+	TArray<FTrack>* Tracks = nullptr;
+
 	switch (TrackType)
 	{
 	case EMediaTrackType::Audio:
-		if (TrackIndex == SelectedAudioTrack)
-		{
-			return true; // already selected
-		}
-
-		UE_LOG(LogWmfMedia, Verbose, TEXT("Tracks %llx: Selecting audio track %i instead of %i (%i tracks)"), this, TrackIndex, SelectedAudioTrack, AudioTracks.Num());
-
-		// disable current track
-		if (SelectedAudioTrack != INDEX_NONE)
-		{
-			UE_LOG(LogWmfMedia, VeryVerbose, TEXT("Tracks %llx: Disabling audio track %i"), this, SelectedAudioTrack);
-
-			PresentationDescriptor->DeselectStream(AudioTracks[SelectedAudioTrack].StreamIndex);
-			SelectedAudioTrack = INDEX_NONE;
-		}
-
-		// enable new track
-		if (TrackIndex != INDEX_NONE)
-		{
-			if (!AudioTracks.IsValidIndex(TrackIndex) || FAILED(PresentationDescriptor->SelectStream(AudioTracks[TrackIndex].StreamIndex)))
-			{
-				return false; // invalid track index or failed to select
-			}
-
-			UE_LOG(LogWmfMedia, VeryVerbose, TEXT("Tracks %llx: Enabled audio track %i"), this, TrackIndex);
-		}
-
-		SelectedAudioTrack = TrackIndex;
+		SelectedTrack = &SelectedAudioTrack;
+		Tracks = &AudioTracks;
 		break;
 
 	case EMediaTrackType::Caption:
-		if (TrackIndex == SelectedCaptionTrack)
-		{
-			return true;
-		}
-
-		UE_LOG(LogWmfMedia, Verbose, TEXT("Tracks %llx: Selecting caption track %i instead of %i (%i tracks)"), this, TrackIndex, SelectedCaptionTrack, CaptionTracks.Num());
-
-		// disable current track
-		if (SelectedCaptionTrack != INDEX_NONE)
-		{
-			UE_LOG(LogWmfMedia, VeryVerbose, TEXT("Tracks %llx: Disabling caption track %i"), this, SelectedCaptionTrack);
-
-			PresentationDescriptor->DeselectStream(CaptionTracks[SelectedCaptionTrack].StreamIndex);
-			SelectedCaptionTrack = INDEX_NONE;
-		}
-
-		// enable new track
-		if (TrackIndex != INDEX_NONE)
-		{
-			if (!CaptionTracks.IsValidIndex(TrackIndex) || FAILED(PresentationDescriptor->SelectStream(CaptionTracks[TrackIndex].StreamIndex)))
-			{
-				return false; // invalid track index or failed to select
-			}
-
-			UE_LOG(LogWmfMedia, VeryVerbose, TEXT("Tracks %llx: Enabled caption track %i"), this, TrackIndex);
-		}
-
-		SelectedCaptionTrack = TrackIndex;
+		SelectedTrack = &SelectedCaptionTrack;
+		Tracks = &CaptionTracks;
 		break;
 
 	case EMediaTrackType::Metadata:
-		if (TrackIndex == SelectedBinaryTrack)
-		{
-			return true; // already selected
-		}
+		SelectedTrack = &SelectedMetadataTrack;
+		Tracks = &MetadataTracks;
+		break;
 
-		UE_LOG(LogWmfMedia, Verbose, TEXT("Tracks %llx: Selecting binary track %i instead of %i (%i tracks)"), this, TrackIndex, SelectedBinaryTrack, MetadataTracks.Num());
+	case EMediaTrackType::Video:
+		SelectedTrack = &SelectedVideoTrack;
+		Tracks = &VideoTracks;
+		break;
 
-		// disable current track
-		if (SelectedBinaryTrack != INDEX_NONE)
-		{
-			UE_LOG(LogWmfMedia, VeryVerbose, TEXT("Tracks %llx: Disabling metadata track %i"), this, SelectedBinaryTrack);
+	default:
+		return false; // unsupported track type
+	}
 
-			PresentationDescriptor->DeselectStream(MetadataTracks[SelectedBinaryTrack].StreamIndex);
-			SelectedBinaryTrack = INDEX_NONE;
-		}
+	check(SelectedTrack != nullptr);
+	check(Tracks != nullptr);
 
-		// enable new track
-		if (TrackIndex != INDEX_NONE)
-		{
-			if (!MetadataTracks.IsValidIndex(TrackIndex) || FAILED(PresentationDescriptor->SelectStream(MetadataTracks[TrackIndex].StreamIndex)))
-			{
-				return false; // invalid track index or failed to select
-			}
+	if (TrackIndex == *SelectedTrack)
+	{
+		return true; // already selected
+	}
 
-			UE_LOG(LogWmfMedia, VeryVerbose, TEXT("Tracks %llx: Enabled metadata track %i"), this, TrackIndex);
-		}
+	if ((TrackIndex != INDEX_NONE) && !Tracks->IsValidIndex(TrackIndex))
+	{
+		return false; // invalid track
+	}
 
-		SelectedBinaryTrack = TrackIndex;
-		break;
+	// deselect stream for old track
+	if (*SelectedTrack != INDEX_NONE)
+	{
+		const DWORD StreamIndex = (*Tracks)[*SelectedTrack].StreamIndex;
+		const HRESULT Result = PresentationDescriptor->DeselectStream(StreamIndex);
 
-	case EMediaTrackType::Video:
-		if (TrackIndex == SelectedVideoTrack)
+		if (FAILED(Result))
 		{
-			return true; // already selected
+			UE_LOG(LogWmfMedia, Verbose, TEXT("Tracks %p: Failed to deselect stream %i on presentation descriptor: %s"), this, StreamIndex, *WmfMedia::ResultToString(Result));
+			return false;
 		}
 
-		UE_LOG(LogWmfMedia, Verbose, TEXT("Tracks %llx: Selecting video track %i instead of %i (%i tracks)"), this, TrackIndex, SelectedVideoTrack, VideoTracks.Num());
-
-		// disable current track
-		if (SelectedVideoTrack != INDEX_NONE)
-		{
-			UE_LOG(LogWmfMedia, VeryVerbose, TEXT("Tracks %llx: Disabling video track %i"), this, SelectedVideoTrack);
+		UE_LOG(LogWmfMedia, Verbose, TEXT("Tracks %p: Disabled stream %i"), this, StreamIndex);
+			
+		*SelectedTrack = INDEX_NONE;
+		SelectionChanged = true;
+	}
 
-			PresentationDescriptor->DeselectStream(VideoTracks[SelectedVideoTrack].StreamIndex);
-			SelectedVideoTrack = INDEX_NONE;
-		}
+	// select stream for new track
+	if (TrackIndex != INDEX_NONE)
+	{
+		const DWORD StreamIndex = (*Tracks)[TrackIndex].StreamIndex;
+		const HRESULT Result = PresentationDescriptor->SelectStream(StreamIndex);
 
-		// enable new track
-		if (TrackIndex != INDEX_NONE)
+		if (FAILED(Result))
 		{
-			if (!VideoTracks.IsValidIndex(TrackIndex) || FAILED(PresentationDescriptor->SelectStream(VideoTracks[TrackIndex].StreamIndex)))
-			{
-				return false; // invalid track index or failed to select
-			}
-
-			UE_LOG(LogWmfMedia, VeryVerbose, TEXT("Tracks %llx: Enabled video track %i"), this, TrackIndex);
+			UE_LOG(LogWmfMedia, Verbose, TEXT("Tracks %p: Failed to enable %s track %i (stream %i): %s"), this, *MediaUtils::TrackTypeToString(TrackType), TrackIndex, StreamIndex, *WmfMedia::ResultToString(Result));
+			return false;
 		}
 
-		SelectedVideoTrack = TrackIndex;
-		break;
+		UE_LOG(LogWmfMedia, Verbose, TEXT("Tracks %p: Enabled stream %i"), this, StreamIndex);
 
-	default:
-		return false; // unsupported track type
+		*SelectedTrack = TrackIndex;
+		SelectionChanged = true;
 	}
 
-	TopologyChanged = true;
-
 	return true;
 }
 
 
 bool FWmfMediaTracks::SetTrackFormat(EMediaTrackType TrackType, int32 TrackIndex, int32 FormatIndex)
 {
+	UE_LOG(LogWmfMedia, Verbose, TEXT("Tracks %p: Setting format on %s track %i to %i"), this, *MediaUtils::TrackTypeToString(TrackType), TrackIndex, FormatIndex);
+
 	FScopeLock Lock(&CriticalSection);
 
-	FTrack* Track = nullptr;
+	TArray<FTrack>* Tracks = nullptr;
 
 	switch (TrackType)
 	{
 	case EMediaTrackType::Audio:
-		if (AudioTracks.IsValidIndex(TrackIndex))
-		{
-			Track = &AudioTracks[TrackIndex];
-		}
+		Tracks = &AudioTracks;
 		break;
 
 	case EMediaTrackType::Caption:
-		if (CaptionTracks.IsValidIndex(TrackIndex))
-		{
-			Track = &CaptionTracks[TrackIndex];
-		}
+		Tracks = &CaptionTracks;
 		break;
 
 	case EMediaTrackType::Metadata:
-		if (MetadataTracks.IsValidIndex(TrackIndex))
-		{
-			Track = &MetadataTracks[TrackIndex];
-		}
+		Tracks = &MetadataTracks;
 		break;
 
 	case EMediaTrackType::Video:
-		if (VideoTracks.IsValidIndex(TrackIndex))
-		{
-			Track = &VideoTracks[TrackIndex];
-		}
+		Tracks = &VideoTracks;
 		break;
 
 	default:
-		break; // unsupported track type
+		return false; // unsupported track type
 	};
 
-	if (Track == nullptr)
+	check(Tracks != nullptr);
+
+	if (!Tracks->IsValidIndex(TrackIndex))
 	{
-		return false;
+		return false; // invalid track index
+	}
+
+	FTrack& Track = (*Tracks)[TrackIndex];
+
+	if (Track.SelectedFormat == FormatIndex)
+	{
+		return true; // format already set
 	}
 
-	if (Track->SelectedFormat != FormatIndex)
+	if (!Track.Formats.IsValidIndex(FormatIndex))
 	{
-		Track->SelectedFormat = FormatIndex;
-		TopologyChanged = true;
+		return false; // invalid format index
 	}
 
+	// set track format
+	UE_LOG(LogWmfMedia, Verbose, TEXT("Tracks %p: Set format %i instead of %i on %s track %i (%i formats)"), this, FormatIndex, Track.SelectedFormat, *MediaUtils::TrackTypeToString(TrackType), TrackIndex, Track.Formats.Num());
+
+	Track.SelectedFormat = FormatIndex;
+	SelectionChanged = true;
+
 	return true;
 }
 
 
 bool FWmfMediaTracks::SetVideoTrackFrameRate(int32 TrackIndex, int32 FormatIndex, float FrameRate)
 {
+	UE_LOG(LogWmfMedia, Verbose, TEXT("Tracks %p: Setting frame rate on format %i of video track %i to %f"), this, FormatIndex, TrackIndex, FrameRate);
+
 	FScopeLock Lock(&CriticalSection);
 
 	const FFormat* Format = GetVideoFormat(TrackIndex, FormatIndex);
@@ -948,7 +915,7 @@ bool FWmfMediaTracks::AddTrackToTopology(const FTrack& Track, IMFTopology& Topol
 
 		if (FAILED(Result))
 		{
-			UE_LOG(LogWmfMedia, Warning, TEXT("Failed to set current media type for stream %i (%s)"), Track.StreamIndex, *WmfMedia::ResultToString(Result));
+			UE_LOG(LogWmfMedia, Verbose, TEXT("Tracks %p: Failed to set current media type for stream %i: %s"), this, Track.StreamIndex, *WmfMedia::ResultToString(Result));
 			return false;
 		}
 	}
@@ -963,7 +930,7 @@ bool FWmfMediaTracks::AddTrackToTopology(const FTrack& Track, IMFTopology& Topol
 
 		if (FAILED(Result))
 		{
-			UE_LOG(LogWmfMedia, Warning, TEXT("Failed to create audio renderer for stream %i (%s)"), Track.StreamIndex, *WmfMedia::ResultToString(Result));
+			UE_LOG(LogWmfMedia, Verbose, TEXT("Tracks %p: Failed to create audio renderer for stream %i: %s"), this, Track.StreamIndex, *WmfMedia::ResultToString(Result));
 			return false;
 		}
 
@@ -979,7 +946,7 @@ bool FWmfMediaTracks::AddTrackToTopology(const FTrack& Track, IMFTopology& Topol
 
 				if (FAILED(Result))
 				{
-					UE_LOG(LogWmfMedia, Warning, TEXT("Failed to override HMD audio output device for stream %i (%s)"), Track.StreamIndex, *WmfMedia::ResultToString(Result));
+					UE_LOG(LogWmfMedia, Verbose, TEXT("Tracks %p: Failed to override HMD audio output device for stream %i: %s"), this, Track.StreamIndex, *WmfMedia::ResultToString(Result));
 					return false;
 				}
 			}
@@ -1004,7 +971,7 @@ bool FWmfMediaTracks::AddTrackToTopology(const FTrack& Track, IMFTopology& Topol
 		{
 			Sampler->OnSample().AddRaw(this, &FWmfMediaTracks::HandleMediaSamplerMetadataSample);
 		}
-		else if ((MajorType == MFMediaType_Video) || (MajorType == MFMediaType_Image))
+		else if (MajorType == MFMediaType_Video)
 		{
 			Sampler->OnSample().AddRaw(this, &FWmfMediaTracks::HandleMediaSamplerVideoSample);
 		}
@@ -1013,7 +980,7 @@ bool FWmfMediaTracks::AddTrackToTopology(const FTrack& Track, IMFTopology& Topol
 
 		if (FAILED(Result))
 		{
-			UE_LOG(LogWmfMedia, Warning, TEXT("Failed to create sampler grabber sink for stream %i (%s)"), Track.StreamIndex, *WmfMedia::ResultToString(Result));
+			UE_LOG(LogWmfMedia, Verbose, TEXT("Tracks %p: Failed to create sampler grabber sink for stream %i: %s"), this, Track.StreamIndex, *WmfMedia::ResultToString(Result));
 			return false;
 		}
 	}
@@ -1032,7 +999,7 @@ bool FWmfMediaTracks::AddTrackToTopology(const FTrack& Track, IMFTopology& Topol
 		FAILED(OutputNode->SetUINT32(MF_TOPONODE_NOSHUTDOWN_ON_REMOVE, FALSE)) ||
 		FAILED(Topology.AddNode(OutputNode)))
 	{
-		UE_LOG(LogWmfMedia, Warning, TEXT("Failed to configure output node for stream %i"), Track.StreamIndex);
+		UE_LOG(LogWmfMedia, Verbose, TEXT("Tracks %p: Failed to configure output node for stream %i"), this, Track.StreamIndex);
 		return false;
 	}
 
@@ -1045,7 +1012,7 @@ bool FWmfMediaTracks::AddTrackToTopology(const FTrack& Track, IMFTopology& Topol
 		FAILED(SourceNode->SetUnknown(MF_TOPONODE_STREAM_DESCRIPTOR, Track.Descriptor)) ||
 		FAILED(Topology.AddNode(SourceNode)))
 	{
-		UE_LOG(LogWmfMedia, Warning, TEXT("Failed to configure source node for stream %i"), Track.StreamIndex);
+		UE_LOG(LogWmfMedia, Verbose, TEXT("Tracks %p: Failed to configure source node for stream %i"), this, Track.StreamIndex);
 		return false;
 	}
 
@@ -1054,7 +1021,7 @@ bool FWmfMediaTracks::AddTrackToTopology(const FTrack& Track, IMFTopology& Topol
 
 	if (FAILED(Result))
 	{
-		UE_LOG(LogWmfMedia, Warning, TEXT("Failed to connect topology nodes for stream %i (%s)"), Track.StreamIndex, *WmfMedia::ResultToString(Result));
+		UE_LOG(LogWmfMedia, Verbose, TEXT("Tracks %p: Failed to connect topology nodes for stream %i: %s"), this, Track.StreamIndex, *WmfMedia::ResultToString(Result));
 		return false;
 	}
 
@@ -1062,7 +1029,7 @@ bool FWmfMediaTracks::AddTrackToTopology(const FTrack& Track, IMFTopology& Topol
 }
 
 
-void FWmfMediaTracks::AddStreamToTracks(uint32 StreamIndex, FString& OutInfo)
+bool FWmfMediaTracks::AddStreamToTracks(uint32 StreamIndex, FString& OutInfo)
 {
 	OutInfo += FString::Printf(TEXT("Stream %i\n"), StreamIndex);
 
@@ -1070,19 +1037,24 @@ void FWmfMediaTracks::AddStreamToTracks(uint32 StreamIndex, FString& OutInfo)
 	TComPtr<IMFStreamDescriptor> StreamDescriptor;
 	{
 		BOOL Selected = FALSE;
-		const HRESULT Result = PresentationDescriptor->GetStreamDescriptorByIndex(StreamIndex, &Selected, &StreamDescriptor);
+		HRESULT Result = PresentationDescriptor->GetStreamDescriptorByIndex(StreamIndex, &Selected, &StreamDescriptor);
 
 		if (FAILED(Result))
 		{
-			UE_LOG(LogWmfMedia, Warning, TEXT("Skipping missing stream descriptor for stream %i (%s)"), StreamIndex, *WmfMedia::ResultToString(Result));
-			OutInfo += TEXT("    missing stream descriptor\n");
+			UE_LOG(LogWmfMedia, Verbose, TEXT("Tracks %p: Failed to get stream descriptor for stream %i: %s"), this, StreamIndex, *WmfMedia::ResultToString(Result));
+			OutInfo += TEXT("\tmissing stream descriptor\n");
 
-			return;
+			return false;
 		}
 
 		if (Selected == TRUE)
 		{
-			PresentationDescriptor->DeselectStream(StreamIndex);
+			Result = PresentationDescriptor->DeselectStream(StreamIndex);
+
+			if (FAILED(Result))
+			{
+				UE_LOG(LogWmfMedia, Verbose, TEXT("Tracks %p: Failed to deselect stream %i: %s"), this, StreamIndex, *WmfMedia::ResultToString(Result));
+			}
 		}
 	}
 
@@ -1093,10 +1065,10 @@ void FWmfMediaTracks::AddStreamToTracks(uint32 StreamIndex, FString& OutInfo)
 
 		if (FAILED(Result))
 		{
-			UE_LOG(LogWmfMedia, Warning, TEXT("Failed to get media type handler for stream %i (%s)"), StreamIndex, *WmfMedia::ResultToString(Result));
-			OutInfo += TEXT("    no handler available\n");
+			UE_LOG(LogWmfMedia, Verbose, TEXT("Tracks %p: Failed to get media type handler for stream %i: %s"), this, StreamIndex, *WmfMedia::ResultToString(Result));
+			OutInfo += TEXT("\tno handler available\n");
 
-			return;
+			return false;
 		}
 	}
 
@@ -1107,23 +1079,23 @@ void FWmfMediaTracks::AddStreamToTracks(uint32 StreamIndex, FString& OutInfo)
 
 		if (FAILED(Result))
 		{
-			UE_LOG(LogWmfMedia, Warning, TEXT("Failed to determine major type of stream %i (%s)"), StreamIndex, *WmfMedia::ResultToString(Result));
-			OutInfo += TEXT("    failed to determine MajorType\n");
+			UE_LOG(LogWmfMedia, Verbose, TEXT("Tracks %p: Failed to determine major type of stream %i: %s"), this, StreamIndex, *WmfMedia::ResultToString(Result));
+			OutInfo += TEXT("\tfailed to determine MajorType\n");
 
-			return;
+			return false;
 		}
 
-		UE_LOG(LogWmfMedia, Verbose, TEXT("Tracks %llx: Major type of stream %i is %s"), this, StreamIndex, *WmfMedia::MajorTypeToString(MajorType));
-		OutInfo += FString::Printf(TEXT("    Type: %s\n"), *WmfMedia::MajorTypeToString(MajorType));
+		UE_LOG(LogWmfMedia, Verbose, TEXT("Tracks %p: Major type of stream %i is %s"), this, StreamIndex, *WmfMedia::MajorTypeToString(MajorType));
+		OutInfo += FString::Printf(TEXT("\tType: %s\n"), *WmfMedia::MajorTypeToString(MajorType));
 
 		if ((MajorType != MFMediaType_Audio) &&
 			(MajorType != MFMediaType_SAMI) &&
 			(MajorType != MFMediaType_Video))
 		{
-			UE_LOG(LogWmfMedia, Warning, TEXT("Unsupported major type %s for stream %i"), *WmfMedia::MajorTypeToString(MajorType), StreamIndex);
-			OutInfo += TEXT("    Unsupported stream type\n");
+			UE_LOG(LogWmfMedia, Verbose, TEXT("Tracks %p: Unsupported major type %s of stream %i"), this, *WmfMedia::MajorTypeToString(MajorType), StreamIndex);
+			OutInfo += TEXT("\tUnsupported stream type\n");
 
-			return;
+			return false;
 		}
 	}
 
@@ -1132,7 +1104,7 @@ void FWmfMediaTracks::AddStreamToTracks(uint32 StreamIndex, FString& OutInfo)
 	{
 		if (Protected)
 		{
-			OutInfo += FString::Printf(TEXT("    Protected content\n"));
+			OutInfo += FString::Printf(TEXT("\tProtected content\n"));
 		}
 	}
 
@@ -1143,45 +1115,55 @@ void FWmfMediaTracks::AddStreamToTracks(uint32 StreamIndex, FString& OutInfo)
 
 		if (FAILED(Result))
 		{
-			UE_LOG(LogWmfMedia, Warning, TEXT("Failed to get number of track formats in stream %i of type %s"), StreamIndex, *WmfMedia::MajorTypeToString(MajorType));
-			OutInfo += TEXT("    failed to get track formats\n");
+			UE_LOG(LogWmfMedia, Verbose, TEXT("Tracks %p: Failed to get number of track formats in stream %i"), this, StreamIndex);
+			OutInfo += TEXT("\tfailed to get track formats\n");
 
-			return;
+			return false;
 		}
 	}
 
 	// create & add track
 	FTrack* Track = nullptr;
+	int32 TrackIndex = INDEX_NONE;
+	int32* SelectedTrack = nullptr;
 
 	if (MajorType == MFMediaType_Audio)
 	{
-		const int32 TrackIndex = AudioTracks.AddDefaulted();
+		SelectedTrack = &SelectedAudioTrack;
+		TrackIndex = AudioTracks.AddDefaulted();
 		Track = &AudioTracks[TrackIndex];
 	}
 	else if (MajorType == MFMediaType_SAMI)
 	{
-		const int32 TrackIndex = CaptionTracks.AddDefaulted();
+		SelectedTrack = &SelectedCaptionTrack;
+		TrackIndex = CaptionTracks.AddDefaulted();
 		Track = &CaptionTracks[TrackIndex];
 	}
 	else if (MajorType == MFMediaType_Binary)
 	{
-		const int32 TrackIndex = MetadataTracks.AddDefaulted();
+		SelectedTrack = &SelectedMetadataTrack;
+		TrackIndex = MetadataTracks.AddDefaulted();
 		Track = &MetadataTracks[TrackIndex];
 	}
-	else if ((MajorType == MFMediaType_Video) || (MajorType == MFMediaType_Image))
+	else if (MajorType == MFMediaType_Video)
 	{
-		const int32 TrackIndex = VideoTracks.AddDefaulted();
+		SelectedTrack = &SelectedVideoTrack;
+		TrackIndex = VideoTracks.AddDefaulted();
 		Track = &VideoTracks[TrackIndex];
 	}
 
 	check(Track != nullptr);
+	check(TrackIndex != INDEX_NONE);
+	check(SelectedTrack != nullptr);
 
 	// get current format
 	TComPtr<IMFMediaType> CurrentMediaType;
 	{
-		if (FAILED(Handler->GetCurrentMediaType(&CurrentMediaType)))
+		HRESULT Result = Handler->GetCurrentMediaType(&CurrentMediaType);
+
+		if (FAILED(Result))
 		{
-			UE_LOG(LogWmfMedia, Warning, TEXT("Failed to get current media type in stream %i of type %s"), StreamIndex, *WmfMedia::MajorTypeToString(MajorType));
+			UE_LOG(LogWmfMedia, Verbose, TEXT("Tracks %p: Failed to get current media type in stream %i: %s"), this, StreamIndex, *WmfMedia::ResultToString(Result));
 		}
 	}
 
@@ -1192,14 +1174,14 @@ void FWmfMediaTracks::AddStreamToTracks(uint32 StreamIndex, FString& OutInfo)
 
 	for (DWORD TypeIndex = 0; TypeIndex < NumMediaTypes; ++TypeIndex)
 	{
-		OutInfo += FString::Printf(TEXT("    Format %i\n"), TypeIndex);
+		OutInfo += FString::Printf(TEXT("\tFormat %i\n"), TypeIndex);
 
 		// get media type
 		TComPtr<IMFMediaType> MediaType;
 
 		if (FAILED(Handler->GetMediaTypeByIndex(TypeIndex, &MediaType)))
 		{
-			OutInfo += TEXT("        failed to get media type\n");
+			OutInfo += TEXT("\t\tfailed to get media type\n");
 
 			continue;
 		}
@@ -1217,22 +1199,22 @@ void FWmfMediaTracks::AddStreamToTracks(uint32 StreamIndex, FString& OutInfo)
 
 			if (FAILED(Result))
 			{
-				UE_LOG(LogWmfMedia, Warning, TEXT("Failed to get sub-type of format %i in stream %i (%s)"), TypeIndex, StreamIndex, *WmfMedia::ResultToString(Result));
-				OutInfo += TEXT("        failed to get sub-type\n");
+				UE_LOG(LogWmfMedia, Verbose, TEXT("Tracks %p: Failed to get sub-type of format %i in stream %i: %s"), this, TypeIndex, StreamIndex, *WmfMedia::ResultToString(Result));
+				OutInfo += TEXT("\t\tfailed to get sub-type\n");
 
 				continue;;
 			}
 		}
 
 		const FString TypeName = WmfMedia::SubTypeToString(SubType);
-		OutInfo += FString::Printf(TEXT("        Codec: %s\n"), *TypeName);
+		OutInfo += FString::Printf(TEXT("\t\tCodec: %s\n"), *TypeName);
 
 		// create output type
 		TComPtr<IMFMediaType> OutputType = WmfMedia::CreateOutputType(MajorType, SubType, AllowNonStandardCodecs);
 
 		if (!OutputType.IsValid())
 		{
-			OutInfo += TEXT("        failed to create output type\n");
+			OutInfo += TEXT("\t\tfailed to create output type\n");
 
 			continue;
 		}
@@ -1258,9 +1240,9 @@ void FWmfMediaTracks::AddStreamToTracks(uint32 StreamIndex, FString& OutInfo)
 				{ 0 }
 			});
 
-			OutInfo += FString::Printf(TEXT("        Channels: %i\n"), NumChannels);
-			OutInfo += FString::Printf(TEXT("        Sample Rate: %i Hz\n"), SampleRate);
-			OutInfo += FString::Printf(TEXT("        Bits Per Sample: %i\n"), BitsPerSample);
+			OutInfo += FString::Printf(TEXT("\t\tChannels: %i\n"), NumChannels);
+			OutInfo += FString::Printf(TEXT("\t\tSample Rate: %i Hz\n"), SampleRate);
+			OutInfo += FString::Printf(TEXT("\t\tBits Per Sample: %i\n"), BitsPerSample);
 		}
 		else if (MajorType == MFMediaType_SAMI)
 		{
@@ -1290,8 +1272,8 @@ void FWmfMediaTracks::AddStreamToTracks(uint32 StreamIndex, FString& OutInfo)
 
 				if (FAILED(Result))
 				{
-					UE_LOG(LogWmfMedia, Error, TEXT("Failed to get video output sub-type for stream %i (%s)"), StreamIndex, *WmfMedia::ResultToString(Result));
-					OutInfo += FString::Printf(TEXT("        failed to get sub-type"));
+					UE_LOG(LogWmfMedia, Verbose, TEXT("Tracks %p: Failed to get video output sub-type for stream %i: %s"), this, StreamIndex, *WmfMedia::ResultToString(Result));
+					OutInfo += FString::Printf(TEXT("\t\tfailed to get sub-type"));
 
 					continue;
 				}
@@ -1303,12 +1285,12 @@ void FWmfMediaTracks::AddStreamToTracks(uint32 StreamIndex, FString& OutInfo)
 			{
 				if (SUCCEEDED(::MFGetAttributeSize(MediaType, MF_MT_FRAME_SIZE, (UINT32*)&OutputDim.X, (UINT32*)&OutputDim.Y)))
 				{
-					OutInfo += FString::Printf(TEXT("        Dimensions: %i x %i\n"), OutputDim.X, OutputDim.Y);
+					OutInfo += FString::Printf(TEXT("\t\tDimensions: %i x %i\n"), OutputDim.X, OutputDim.Y);
 				}
 				else
 				{
 					OutputDim = FIntPoint::ZeroValue;
-					OutInfo += FString::Printf(TEXT("        Dimensions: n/a\n"));
+					OutInfo += FString::Printf(TEXT("\t\tDimensions: n/a\n"));
 				}
 			}
 
@@ -1320,12 +1302,12 @@ void FWmfMediaTracks::AddStreamToTracks(uint32 StreamIndex, FString& OutInfo)
 				if (SUCCEEDED(::MFGetAttributeRatio(MediaType, MF_MT_FRAME_RATE, &Numerator, &Denominator)))
 				{
 					FrameRate = static_cast<float>(Numerator) / Denominator;
-					OutInfo += FString::Printf(TEXT("        Frame Rate: %g fps\n"), FrameRate);
+					OutInfo += FString::Printf(TEXT("\t\tFrame Rate: %g fps\n"), FrameRate);
 				}
 				else
 				{
 					FrameRate = 0.0f;
-					OutInfo += FString::Printf(TEXT("        Frame Rate: n/a\n"));
+					OutInfo += FString::Printf(TEXT("\t\tFrame Rate: n/a\n"));
 				}
 			}
 
@@ -1355,7 +1337,12 @@ void FWmfMediaTracks::AddStreamToTracks(uint32 StreamIndex, FString& OutInfo)
 					FrameRates = TRange<float>(FrameRate);
 				}
 
-				OutInfo += FString::Printf(TEXT("        Frame Rate Range: %g - %g fps\n"), FrameRates.GetLowerBoundValue(), FrameRates.GetUpperBoundValue());
+				OutInfo += FString::Printf(TEXT("\t\tFrame Rate Range: %g - %g fps\n"), FrameRates.GetLowerBoundValue(), FrameRates.GetUpperBoundValue());
+
+				if (FrameRates.IsDegenerate() && (FrameRates.GetLowerBoundValue() == 1.0f))
+				{
+					OutInfo += FString::Printf(TEXT("\t\tpossibly a still image stream (may not work)\n"));
+				}
 			}
 
 			FIntPoint BufferDim;
@@ -1488,7 +1475,41 @@ void FWmfMediaTracks::AddStreamToTracks(uint32 StreamIndex, FString& OutInfo)
 		}
 	}
 
-	// get stream info
+	// ensure that a track format is selected
+	if (Track->SelectedFormat == INDEX_NONE)
+	{
+		for (int32 FormatIndex = 0; FormatIndex < Track->Formats.Num(); ++FormatIndex)
+		{
+			const FFormat& Format = Track->Formats[FormatIndex];
+			const HRESULT Result = Handler->SetCurrentMediaType(Format.InputType);
+
+			if (SUCCEEDED(Result))
+			{
+				UE_LOG(LogWmfMedia, Verbose, TEXT("Tracks %p: Picked default format %i for stream %i"), this, FormatIndex, StreamIndex);
+				Track->SelectedFormat = FormatIndex;
+				break;
+			}
+		}
+
+		if (Track->SelectedFormat == INDEX_NONE)
+		{
+			UE_LOG(LogWmfMedia, Verbose, TEXT("Tracks %p: No supported media types found in stream %i"), this, StreamIndex);
+			OutInfo += TEXT("\tunsupported media type\n");
+		}
+	}
+
+	// select this track if no track is selected yet
+	if ((*SelectedTrack == INDEX_NONE) && (Track->SelectedFormat != INDEX_NONE))
+	{
+		const HRESULT Result = PresentationDescriptor->SelectStream(StreamIndex);
+
+		if (SUCCEEDED(Result))
+		{
+			*SelectedTrack = TrackIndex;
+		}
+	}
+
+	// set track details
 	PWSTR OutString = NULL;
 	UINT32 OutLength = 0;
 
@@ -1513,17 +1534,7 @@ void FWmfMediaTracks::AddStreamToTracks(uint32 StreamIndex, FString& OutInfo)
 	Track->Protected = Protected;
 	Track->StreamIndex = StreamIndex;
 
-	if (Track->Formats.Num() > 0)
-	{
-		Track->SelectedFormat = 0;
-	}
-	else
-	{
-		Track->SelectedFormat = INDEX_NONE;
-
-		UE_LOG(LogWmfMedia, Warning, TEXT("No supported media type in stream %i of type %s"), StreamIndex, *WmfMedia::MajorTypeToString(MajorType));
-		OutInfo += TEXT("    unsupported media type\n");
-	}
+	return true;
 }
 
 
@@ -1674,13 +1685,13 @@ void FWmfMediaTracks::HandleMediaSamplerMetadataSample(const uint8* Buffer, uint
 
 	FScopeLock Lock(&CriticalSection);
 
-	if (!MetadataTracks.IsValidIndex(SelectedBinaryTrack))
+	if (!MetadataTracks.IsValidIndex(SelectedMetadataTrack))
 	{
 		return; // invalid track index
 	}
 
 	// create & add sample to queue
-	const FTrack& Track = MetadataTracks[SelectedBinaryTrack];
+	const FTrack& Track = MetadataTracks[SelectedMetadataTrack];
 	const auto BinarySample = MakeShared<FWmfMediaBinarySample, ESPMode::ThreadSafe>();
 
 	if (BinarySample->Initialize(Buffer, Size, Time, Duration))
diff --git a/Engine/Plugins/Media/WmfMedia/Source/WmfMedia/Private/Player/WmfMediaTracks.h b/Engine/Plugins/Media/WmfMedia/Source/WmfMedia/Private/Player/WmfMediaTracks.h
--- a/Engine/Plugins/Media/WmfMedia/Source/WmfMedia/Private/Player/WmfMediaTracks.h
+++ b/Engine/Plugins/Media/WmfMedia/Source/WmfMedia/Private/Player/WmfMediaTracks.h
@@ -125,10 +125,10 @@ public:
 	 * Get the current flags.
 	 *
 	 * @param OutMediaSourceChanged Will indicate whether the media source changed.
-	 * @param OutTopologyChanged Will indicate whether the topology changed.
+	 * @param OutSelectionChanged Will indicate whether the track selection changed.
 	 * @see ClearFlags
 	 */
-	void GetFlags(bool& OutMediaSourceChanged, bool& OutTopologyChanged) const;
+	void GetFlags(bool& OutMediaSourceChanged, bool& OutSelectionChanged) const;
 
 	/**
 	 * Get the information string for the currently loaded media source.
@@ -201,8 +201,10 @@ protected:
 	 *
 	 * @param StreamIndex The index of the stream to add.
 	 * @param OutInfo Will contain appended debug information.
+	 * @return true on success, false otherwise.
+	 * @see AddTrackToTopology
 	 */
-	void AddStreamToTracks(uint32 StreamIndex, FString& OutInfo);
+	bool AddStreamToTracks(uint32 StreamIndex, FString& OutInfo);
 
 	/**
 	 * Add the given track to the specified playback topology.
@@ -210,6 +212,7 @@ protected:
 	 * @param Track The track to add.
 	 * @param Topology The playback topology.
 	 * @return true on success, false otherwise.
+	 * @see AddStreamToTracks
 	 */
 	bool AddTrackToTopology(const FTrack& Track, IMFTopology& Topology) const;
 
@@ -302,17 +305,17 @@ private:
 	/** Index of the selected audio track. */
 	int32 SelectedAudioTrack;
 
-	/** Index of the selected binary track. */
-	int32 SelectedBinaryTrack;
-
 	/** Index of the selected caption track. */
 	int32 SelectedCaptionTrack;
 
+	/** Index of the selected binary track. */
+	int32 SelectedMetadataTrack;
+
 	/** Index of the selected video track. */
 	int32 SelectedVideoTrack;
 
-	/** Whether the playback topology changed. */
-	bool TopologyChanged;
+	/** Whether the track selection changed. */
+	bool SelectionChanged;
 
 	/** Video sample object pool. */
 	FWmfMediaTextureSamplePool* VideoSamplePool;
diff --git a/Engine/Plugins/Media/WmfMedia/Source/WmfMedia/Private/Wmf/WmfMediaSession.cpp b/Engine/Plugins/Media/WmfMedia/Source/WmfMedia/Private/Wmf/WmfMediaSession.cpp
--- a/Engine/Plugins/Media/WmfMedia/Source/WmfMedia/Private/Wmf/WmfMediaSession.cpp
+++ b/Engine/Plugins/Media/WmfMedia/Source/WmfMedia/Private/Wmf/WmfMediaSession.cpp
@@ -12,11 +12,17 @@
 #include "AllowWindowsPlatformTypes.h"
 
 
+#define WMFMEDIASESSION_USE_WINDOWS7FASTFORWARDENDHACK 1
+
+
 /* Local helpers
  *****************************************************************************/
 
 namespace WmfMediaSession
 {
+	/** Time span value for RequestedTime indicating a seek to the current time. */
+	const FTimespan RequestedTimeCurrent = FTimespan::MaxValue();
+
 	/** Get the human readable string representation of a media player state. */
 	const TCHAR* StateToString(EMediaState State)
 	{
@@ -25,6 +31,7 @@ namespace WmfMediaSession
 		case EMediaState::Closed: return TEXT("Closed");
 		case EMediaState::Error: return TEXT("Error");
 		case EMediaState::Paused: return TEXT("Paused");
+		case EMediaState::Playing: return TEXT("Playing");
 		case EMediaState::Preparing: return TEXT("Preparing");
 		case EMediaState::Stopped: return TEXT("Stopped");
 		default: return TEXT("Unknown");
@@ -40,16 +47,16 @@ FWmfMediaSession::FWmfMediaSession()
 	: CanScrub(false)
 	, Capabilities(0)
 	, CurrentDuration(FTimespan::Zero())
-	, CurrentRate(0.0f)
-	, CurrentState(EMediaState::Closed)
 	, LastTime(FTimespan::Zero())
-	, PendingChange(false)
+	, PendingChanges(false)
 	, RefCount(0)
-	, RequestedRate(0.0f)
-	, RequestedTime(FTimespan::MinValue())
+	, SessionRate(0.0f)
+	, SessionState(EMediaState::Closed)
 	, ShouldLoop(false)
 	, Status(EMediaStatus::None)
-{ }
+{
+	UE_LOG(LogWmfMedia, Verbose, TEXT("Session %p: Created"), this);
+}
 
 
 FWmfMediaSession::~FWmfMediaSession()
@@ -57,7 +64,7 @@ FWmfMediaSession::~FWmfMediaSession()
 	check(RefCount == 0);
 	Shutdown();
 
-	UE_LOG(LogWmfMedia, Verbose, TEXT("Session %llx: Destroyed"), this);
+	UE_LOG(LogWmfMedia, Verbose, TEXT("Session %p: Destroyed"), this);
 }
 
 
@@ -66,6 +73,25 @@ FWmfMediaSession::~FWmfMediaSession()
 
 void FWmfMediaSession::GetEvents(TArray<EMediaEvent>& OutEvents)
 {
+#if WMFMEDIASESSION_USE_WINDOWS7FASTFORWARDENDHACK
+	if (CurrentDuration > FTimespan::Zero())
+	{
+		FTimespan Time = GetTime();
+
+		if ((Time < FTimespan::Zero()) || (Time > CurrentDuration))
+		{
+			if (ShouldLoop)
+			{
+				HandleSessionEnded();
+			}
+			else
+			{
+				MediaSession->Stop();
+			}
+		}
+	}
+#endif
+
 	EMediaEvent Event;
 
 	while (DeferredEvents.Dequeue(Event))
@@ -79,6 +105,8 @@ bool FWmfMediaSession::Initialize(bool LowLatency)
 {
 	Shutdown();
 
+	UE_LOG(LogWmfMedia, Verbose, TEXT("Session %p: Initializing media session (LowLatency: %d)"), this, LowLatency);
+
 	// create session attributes
 	TComPtr<IMFAttributes> Attributes;
 	{
@@ -86,7 +114,7 @@ bool FWmfMediaSession::Initialize(bool LowLatency)
 
 		if (FAILED(Result))
 		{
-			UE_LOG(LogWmfMedia, Error, TEXT("Failed to create media session attributes (%s)"), *WmfMedia::ResultToString(Result));
+			UE_LOG(LogWmfMedia, Verbose, TEXT("Session %p: Failed to create media session attributes: %s"), this, *WmfMedia::ResultToString(Result));
 			return false;
 		}
 	}
@@ -99,21 +127,23 @@ bool FWmfMediaSession::Initialize(bool LowLatency)
 
 			if (FAILED(Result))
 			{
-				UE_LOG(LogWmfMedia, Warning, TEXT("Failed to set low latency session attribute (%s)"), *WmfMedia::ResultToString(Result));
+				UE_LOG(LogWmfMedia, Verbose, TEXT("Session %p: Failed to set low latency session attribute: %s"), this, *WmfMedia::ResultToString(Result));
 			}
 		}
 		else
 		{
-			UE_LOG(LogWmfMedia, Warning, TEXT("Low latency media processing requires Windows 8 or newer"));
+			UE_LOG(LogWmfMedia, Verbose, TEXT("Session %p: Low latency media processing requires Windows 8 or newer"), this);
 		}
 	}
 
+	FScopeLock Lock(&CriticalSection);
+
 	// create media session
 	HRESULT Result = ::MFCreateMediaSession(Attributes, &MediaSession);
 
 	if (FAILED(Result))
 	{
-		UE_LOG(LogWmfMedia, Error, TEXT("Failed to create media session (%s)"), *WmfMedia::ResultToString(Result));
+		UE_LOG(LogWmfMedia, Verbose, TEXT("Session %p: Failed to create media session: %s"), this, *WmfMedia::ResultToString(Result));
 		return false;
 	}
 
@@ -122,11 +152,11 @@ bool FWmfMediaSession::Initialize(bool LowLatency)
 
 	if (FAILED(Result))
 	{
-		UE_LOG(LogWmfMedia, Error, TEXT("Failed to start media session event processing (%s)"), *WmfMedia::ResultToString(Result));
+		UE_LOG(LogWmfMedia, Verbose, TEXT("Session %p: Failed to start media session event processing: %s"), this, *WmfMedia::ResultToString(Result));
 		return false;
 	}
 
-	CurrentState = EMediaState::Preparing;
+	SessionState = EMediaState::Preparing;
 
 	return true;
 }
@@ -139,11 +169,11 @@ bool FWmfMediaSession::SetTopology(const TComPtr<IMFTopology>& InTopology, FTime
 		return false;
 	}
 
-	UE_LOG(LogWmfMedia, Verbose, TEXT("Session: %llx: Setting topology %llx (duration = %s)"), this, InTopology.Get(), *InDuration.ToString());
+	UE_LOG(LogWmfMedia, Verbose, TEXT("Session: %p: Setting new partial topology (duration = %s)"), this, *InDuration.ToString());
 
 	FScopeLock Lock(&CriticalSection);
 
-	if (CurrentState == EMediaState::Preparing)
+	if (SessionState == EMediaState::Preparing)
 	{
 		// media source resolved
 		if (InTopology != NULL)
@@ -153,9 +183,9 @@ bool FWmfMediaSession::SetTopology(const TComPtr<IMFTopology>& InTopology, FTime
 
 			if (FAILED(Result))
 			{
-				UE_LOG(LogWmfMedia, Error, TEXT("Failed to set topology: %s"), *WmfMedia::ResultToString(Result));
+				UE_LOG(LogWmfMedia, Verbose, TEXT("Session %p: Failed to set topology: %s"), this, *WmfMedia::ResultToString(Result));
 				
-				CurrentState = EMediaState::Error;
+				SessionState = EMediaState::Error;
 				DeferredEvents.Enqueue(EMediaEvent::MediaOpenFailed);
 			}
 			else
@@ -166,15 +196,17 @@ bool FWmfMediaSession::SetTopology(const TComPtr<IMFTopology>& InTopology, FTime
 		else
 		{
 			// no tracks selected
-			CurrentState = EMediaState::Stopped;
+			UpdateCharacteristics();
+			SessionState = EMediaState::Stopped;
 			DeferredEvents.Enqueue(EMediaEvent::MediaOpened);
 		}
 	}
 	else
 	{
 		// topology changed during playback, i.e. track switching
-		if (PendingChange)
+		if (PendingChanges)
 		{
+			UE_LOG(LogWmfMedia, Verbose, TEXT("Session %p: Requesting topology change after pending command"), this);
 			RequestedTopology = InTopology;
 		}
 		else
@@ -196,9 +228,11 @@ void FWmfMediaSession::Shutdown()
 		return;
 	}
 
+	UE_LOG(LogWmfMedia, Verbose, TEXT("Session: %p: Shutting down media session"), this);
+
 	FScopeLock Lock(&CriticalSection);
 
-	DiscardPendingState();
+	DiscardPendingChanges();
 
 	MediaSession->Close();
 	MediaSession->Shutdown();
@@ -212,10 +246,10 @@ void FWmfMediaSession::Shutdown()
 	CanScrub = false;
 	Capabilities = 0;
 	CurrentDuration = FTimespan::Zero();
-	CurrentRate = 0.0f;
-	CurrentState = EMediaState::Closed;
+	SessionRate = 0.0f;
+	SessionState = EMediaState::Closed;
 	LastTime = FTimespan::Zero();
-	RequestedRate = 0.0f;
+	RequestedRate.Reset();
 	Status = EMediaStatus::None;
 	ThinnedRates.Empty();
 	UnthinnedRates.Empty();
@@ -234,12 +268,12 @@ bool FWmfMediaSession::CanControl(EMediaControl Control) const
 
 	if (Control == EMediaControl::Pause)
 	{
-		return ((CurrentState == EMediaState::Playing) && (((Capabilities & MFSESSIONCAP_PAUSE) != 0) || UnthinnedRates.Contains(0.0f)));
+		return ((SessionState == EMediaState::Playing) && (((Capabilities & MFSESSIONCAP_PAUSE) != 0) || UnthinnedRates.Contains(0.0f)));
 	}
 
 	if (Control == EMediaControl::Resume)
 	{
-		return ((CurrentState != EMediaState::Playing) && UnthinnedRates.Contains(1.0f));
+		return ((SessionState != EMediaState::Playing) && UnthinnedRates.Contains(1.0f));
 	}
 
 	if (Control == EMediaControl::Scrub)
@@ -264,13 +298,18 @@ FTimespan FWmfMediaSession::GetDuration() const
 
 float FWmfMediaSession::GetRate() const
 {
-	return CurrentRate;
+	return (SessionState == EMediaState::Playing) ? SessionRate : 0.0f;
 }
 
 
 EMediaState FWmfMediaSession::GetState() const
 {
-	return CurrentState;
+	if ((SessionState == EMediaState::Playing) && (SessionRate == 0.0f))
+	{
+		return EMediaState::Paused;
+	}
+
+	return SessionState;
 }
 
 
@@ -301,8 +340,9 @@ FTimespan FWmfMediaSession::GetTime() const
 	}
 
 	MFTIME ClockTime;
+	MFTIME SystemTime;
 
-	if (FAILED(PresentationClock->GetTime(&ClockTime)))
+	if (FAILED(PresentationClock->GetCorrelatedTime(0, &ClockTime, &SystemTime)))
 	{
 		return FTimespan::Zero();
 	}
@@ -324,28 +364,33 @@ bool FWmfMediaSession::Seek(const FTimespan& Time)
 		return false;
 	}
 
-	UE_LOG(LogWmfMedia, Verbose, TEXT("Session %llx: Seeking to %s"), this, *Time.ToString());
+	UE_LOG(LogWmfMedia, Verbose, TEXT("Session %p: Seeking to %s"), this, *Time.ToString());
 
 	FScopeLock Lock(&CriticalSection);
 
 	// validate seek
-	if ((CurrentState == EMediaState::Closed) || (CurrentState == EMediaState::Error))
+	if (!CanControl(EMediaControl::Seek))
 	{
-		UE_LOG(LogWmfMedia, Warning, TEXT("Cannot seek while closed or in error state"));
+		UE_LOG(LogWmfMedia, Verbose, TEXT("Session %p: Media source doesn't support seeking"), this);
+		return false;
+	}
+
+	if ((SessionState == EMediaState::Closed) || (SessionState == EMediaState::Error))
+	{
+		UE_LOG(LogWmfMedia, Verbose, TEXT("Session %p: Cannot seek while closed or in error state"), this);
 		return false;
 	}
 
 	if ((Time < FTimespan::Zero()) || (Time > CurrentDuration))
 	{
-		UE_LOG(LogWmfMedia, Warning, TEXT("Invalid seek time %s (media duration is %s)"), *Time.ToString(), *CurrentDuration.ToString());
+		UE_LOG(LogWmfMedia, Verbose, TEXT("Session %p: Invalid seek time %s (media duration is %s)"), this, *Time.ToString(), *CurrentDuration.ToString());
 		return false;
 	}
 
 	// wait for pending changes to complete
-	if (PendingChange)
+	if (PendingChanges)
 	{
-		UE_LOG(LogWmfMedia, Verbose, TEXT("Session %llx: Requesting seek after pending command"), this);
-
+		UE_LOG(LogWmfMedia, Verbose, TEXT("Session %p: Requesting seek after pending command"), this);
 		RequestedTime = Time;
 
 		return true;
@@ -370,44 +415,40 @@ bool FWmfMediaSession::SetRate(float Rate)
 		return false;
 	}
 
+	UE_LOG(LogWmfMedia, Verbose, TEXT("Session %p: Setting rate to %f"), this, Rate);
+
 	FScopeLock Lock(&CriticalSection);
 
-	if (Rate == CurrentRate)
+#if 0
+	if (Rate == SessionRate)
 	{
-		return true; // rate already set
-	}
+		// Note: play rates and play states are handled separately in WMF. The session's initial play
+		// rate after opening a media source is 1.0, even though the session is in the Stopped state.
 
-	EMediaRateThinning Thinning;
+		if ((Rate != 1.0f) || (SessionState == EMediaState::Playing))
+		{
+			return true; // rate already set
+		}
+	}
+#endif
 
 	// validate rate
-	if (UnthinnedRates.Contains(Rate))
-	{
-		UE_LOG(LogWmfMedia, Verbose, TEXT("Session %llx: Setting rate from to %f to %f (unthinned)"), this, CurrentRate, Rate);
-		Thinning = EMediaRateThinning::Unthinned;
-	}
-	else if (ThinnedRates.Contains(Rate))
-	{
-		UE_LOG(LogWmfMedia, Verbose, TEXT("Session %llx: Setting rate from to %f to %f (thinned)"), this, CurrentRate, Rate);
-		Thinning = EMediaRateThinning::Thinned;
-	}
-	else
+	if (!ThinnedRates.Contains(Rate) && !UnthinnedRates.Contains(Rate))
 	{
-		UE_LOG(LogWmfMedia, Warning, TEXT("The rate %f is not supported"), Rate);
+		UE_LOG(LogWmfMedia, Verbose, TEXT("Session %p: The rate %f is not supported"), this, Rate);
 		return false;
 	}
 
 	// wait for pending changes to complete
-	if (PendingChange)
+	if (PendingChanges)
 	{
-		UE_LOG(LogWmfMedia, Verbose, TEXT("Session %llx: Requesting rate change after pending command"), this);
-
+		UE_LOG(LogWmfMedia, Verbose, TEXT("Session %p: Requesting rate change after pending command"), this);
 		RequestedRate = Rate;
-		RequestedThinning = Thinning;
 
 		return true;
 	}
 
-	return CommitRate(Rate, Thinning);
+	return CommitRate(Rate);
 }
 
 
@@ -442,7 +483,7 @@ STDMETHODIMP FWmfMediaSession::Invoke(IMFAsyncResult* AsyncResult)
 
 		if (FAILED(Result))
 		{
-			UE_LOG(LogWmfMedia, Verbose, TEXT("Session %llx: Failed to get event: %s"), this, *WmfMedia::ResultToString(Result));
+			UE_LOG(LogWmfMedia, Verbose, TEXT("Session %p: Failed to get event: %s"), this, *WmfMedia::ResultToString(Result));
 			return S_OK;
 		}
 	}
@@ -453,7 +494,7 @@ STDMETHODIMP FWmfMediaSession::Invoke(IMFAsyncResult* AsyncResult)
 
 		if (FAILED(Result))
 		{
-			UE_LOG(LogWmfMedia, Verbose, TEXT("Session %llx: Failed to get session event type: %s"), this, *WmfMedia::ResultToString(Result));
+			UE_LOG(LogWmfMedia, Verbose, TEXT("Session %p: Failed to get session event type: %s"), this, *WmfMedia::ResultToString(Result));
 			return S_OK;
 		}
 	}
@@ -464,11 +505,11 @@ STDMETHODIMP FWmfMediaSession::Invoke(IMFAsyncResult* AsyncResult)
 
 		if (FAILED(Result))
 		{
-			UE_LOG(LogWmfMedia, Verbose, TEXT("Session %llx: Failed to get event status: %s"), this, *WmfMedia::ResultToString(Result));
+			UE_LOG(LogWmfMedia, Verbose, TEXT("Session %p: Failed to get event status: %s"), this, *WmfMedia::ResultToString(Result));
 		}
 	}
 
-	UE_LOG(LogWmfMedia, Verbose, TEXT("Session %llx: Event [%s] (%s)"), this, *WmfMedia::MediaEventToString(EventType), *WmfMedia::ResultToString(EventStatus));
+	UE_LOG(LogWmfMedia, Verbose, TEXT("Session %p: Event [%s]: %s"), this, *WmfMedia::MediaEventToString(EventType), *WmfMedia::ResultToString(EventStatus));
 
 	// process event
 	switch (EventType)
@@ -483,10 +524,7 @@ STDMETHODIMP FWmfMediaSession::Invoke(IMFAsyncResult* AsyncResult)
 		break;
 
 	case MEError:
-		UE_LOG(LogWmfMedia, Error, TEXT("An error occurred in the media session: %s"), *WmfMedia::ResultToString(EventStatus));
-		CurrentState = EMediaState::Error;
-		DiscardPendingState();
-		MediaSession->Close();
+		HandleError(EventStatus);	
 		break;
 
 	case MEReconnectEnd:
@@ -508,278 +546,35 @@ STDMETHODIMP FWmfMediaSession::Invoke(IMFAsyncResult* AsyncResult)
 		break;
 
 	case MESessionEnded:
-		DeferredEvents.Enqueue(EMediaEvent::PlaybackEndReached);
-
-		if (ShouldLoop)
-		{
-			// loop back to beginning/end
-			RequestedTime = (CurrentRate < 0.0f) ? CurrentDuration : FTimespan::Zero();
-			UpdatePendingState();
-		}
-		else
-		{
-			CurrentRate = RequestedRate = 0.0f;
-			CurrentState = EMediaState::Stopped;
-			LastTime = FTimespan::Zero();
-
-			if (RateControl != NULL)
-			{
-				RateControl->SetRate(TRUE, 0.0f);
-			}
-		}
+		HandleSessionEnded();
 		break;
 
 	case MESessionPaused:
-		if (SUCCEEDED(EventStatus))
-		{
-			CurrentState = EMediaState::Paused;
-			DeferredEvents.Enqueue(EMediaEvent::PlaybackSuspended);
-			UpdatePendingState();
-		}
-		else
-		{
-			DiscardPendingState();
-		}
+		HandleSessionPaused(EventStatus);
 		break;
 
 	case MESessionRateChanged:
-		if (SUCCEEDED(EventStatus))
-		{
-			// cache current rate
-			PROPVARIANT Value;
-			::PropVariantInit(&Value);
-
-			const HRESULT Result = Event->GetValue(&Value);
-
-			if (SUCCEEDED(Result) && (Value.vt == VT_R4))
-			{
-				CurrentRate = Value.fltVal;
-			}
-		}
-		else if (RateControl != NULL)
-		{
-			BOOL Thin = FALSE;
-			RateControl->GetRate(&Thin, &CurrentRate);
-		}
-
-		RequestedRate = CurrentRate;
-		UpdatePendingState();
+		HandleSessionRateChanged(EventStatus, *Event);
 		break;
 
 	case MESessionScrubSampleComplete:
-		DeferredEvents.Enqueue(EMediaEvent::SeekCompleted);
-		UpdatePendingState();
+		HandleSessionScrubSampleComplete();
 		break;
 
 	case MESessionStarted:
-		if (SUCCEEDED(EventStatus))
-		{
-			if (CurrentState == EMediaState::Playing)
-			{
-				DeferredEvents.Enqueue(EMediaEvent::SeekCompleted);
-			}
-
-			CurrentState = (CurrentRate == 0.0f) ? EMediaState::Paused : EMediaState::Playing;
-			DeferredEvents.Enqueue(EMediaEvent::PlaybackResumed);
-
-			if (CurrentState == EMediaState::Paused)
-			{
-				// wait for MESessionScrubSampleComplete
-			}
-			else
-			{
-				UpdatePendingState();
-			}
-		}
-		else
-		{
-			DiscardPendingState();
-		}
+		HandleSessionStarted(EventStatus);
 		break;
 
 	case MESessionStopped:
-		if (SUCCEEDED(EventStatus))
-		{
-			CurrentState = EMediaState::Stopped;
-			DeferredEvents.Enqueue(EMediaEvent::PlaybackSuspended);
-			UpdatePendingState();
-		}
-		else
-		{
-			DiscardPendingState();
-		}
+		HandleSessionStopped(EventStatus);
 		break;
 
 	case MESessionTopologySet:
-		if (SUCCEEDED(EventStatus))
-		{
-			if (CurrentState != EMediaState::Preparing)
-			{
-				// track switching successful
-				RequestedTopology.Reset();
-				UpdatePendingState();
-			}
-		}
-		else
-		{
-			if (CurrentState == EMediaState::Preparing)
-			{
-				// new media has incorrect topology
-				CurrentState = EMediaState::Error;
-				DeferredEvents.Enqueue(EMediaEvent::MediaOpenFailed);
-			}
-
-			DiscardPendingState();
-		}
+		HandleSessionTopologySet(EventStatus, *Event);
 		break;
 
 	case MESessionTopologyStatus:
-		if (SUCCEEDED(EventStatus))
-		{
-			// check if new topology is ready
-			UINT32 TopologyStatus = MF_TOPOSTATUS_INVALID;
-			{
-				const HRESULT Result = Event->GetUINT32(MF_EVENT_TOPOLOGY_STATUS, (UINT32*)&TopologyStatus);
-
-				if (FAILED(Result))
-				{
-					UE_LOG(LogWmfMedia, Verbose, TEXT("Session %llx: Failed to get topology status: %s"), this, *WmfMedia::ResultToString(Result));
-				}
-			}
-
-			// initialize new topology
-			if (TopologyStatus == MF_TOPOSTATUS_READY)
-			{
-				UE_LOG(LogWmfMedia, Verbose, TEXT("Session %llx: Topology ready"), this);
-
-				HRESULT Result = MediaSession->GetFullTopology(MFSESSION_GETFULLTOPOLOGY_CURRENT, 0, &CurrentTopology);
-
-				if (SUCCEEDED(Result))
-				{
-					// get presentation clock, if available
-					TComPtr<IMFClock> Clock;
-
-					Result = MediaSession->GetClock(&Clock);
-
-					if (FAILED(Result))
-					{
-						UE_LOG(LogWmfMedia, Log, TEXT("Session clock unavailable (%s)"), *WmfMedia::ResultToString(Result));
-					}
-					else
-					{
-						Result = Clock->QueryInterface(IID_PPV_ARGS(&PresentationClock));
-
-						if (FAILED(Result))
-						{
-							UE_LOG(LogWmfMedia, Log, TEXT("Presentation clock unavailable (%s)"), *WmfMedia::ResultToString(Result));
-						}
-						else
-						{
-							UE_LOG(LogWmfMedia, Verbose, TEXT("Session %llx: Presentation clock ready"), this);
-						}
-					}
-
-					// get rate control & rate support, if available
-					Result = ::MFGetService(MediaSession, MF_RATE_CONTROL_SERVICE, IID_PPV_ARGS(&RateControl));
-
-					if (FAILED(Result))
-					{
-						UE_LOG(LogWmfMedia, Log, TEXT("Rate control service unavailable (%s)"), *WmfMedia::ResultToString(Result));
-					}
-					else
-					{
-						UE_LOG(LogWmfMedia, Verbose, TEXT("Session %llx: Rate control ready"), this);
-					}
-
-					Result = ::MFGetService(MediaSession, MF_RATE_CONTROL_SERVICE, IID_PPV_ARGS(&RateSupport));
-
-					if (FAILED(Result))
-					{
-						UE_LOG(LogWmfMedia, Log, TEXT("Rate support service unavailable (%s)"), *WmfMedia::ResultToString(Result));
-					}
-					else
-					{
-						UE_LOG(LogWmfMedia, Verbose, TEXT("Session %llx: Rate support ready"), this);
-					}
-
-					// cache rate control properties
-					if (RateSupport.IsValid())
-					{
-						CanScrub = SUCCEEDED(RateSupport->IsRateSupported(TRUE, 0.0f, NULL));
-
-						float MaxRate = 0.0f;
-						float MinRate = 0.0f;
-
-						if (SUCCEEDED(RateSupport->GetSlowestRate(MFRATE_FORWARD, TRUE, &MinRate)) &&
-							SUCCEEDED(RateSupport->GetFastestRate(MFRATE_FORWARD, TRUE, &MaxRate)))
-						{
-							ThinnedRates.Add(TRange<float>::Inclusive(MinRate, MaxRate));
-						}
-
-						if (SUCCEEDED(RateSupport->GetSlowestRate(MFRATE_REVERSE, TRUE, &MinRate)) &&
-							SUCCEEDED(RateSupport->GetFastestRate(MFRATE_REVERSE, TRUE, &MaxRate)))
-						{
-							ThinnedRates.Add(TRange<float>::Inclusive(MaxRate, MinRate));
-						}
-
-						if (SUCCEEDED(RateSupport->GetSlowestRate(MFRATE_FORWARD, FALSE, &MinRate)) &&
-							SUCCEEDED(RateSupport->GetFastestRate(MFRATE_FORWARD, FALSE, &MaxRate)))
-						{
-							UnthinnedRates.Add(TRange<float>::Inclusive(MinRate, MaxRate));
-						}
-
-						if (SUCCEEDED(RateSupport->GetSlowestRate(MFRATE_REVERSE, FALSE, &MinRate)) &&
-							SUCCEEDED(RateSupport->GetFastestRate(MFRATE_REVERSE, FALSE, &MaxRate)))
-						{
-							UnthinnedRates.Add(TRange<float>::Inclusive(MaxRate, MinRate));
-						}
-					}
-					else
-					{
-						CanScrub = false;
-						ThinnedRates.Empty();
-						UnthinnedRates.Empty();
-					}
-
-					// new media opened successfully
-					if (CurrentState == EMediaState::Preparing)
-					{
-						CurrentState = EMediaState::Stopped;
-						DeferredEvents.Enqueue(EMediaEvent::MediaOpened);
-
-						UpdatePendingState();
-					}
-				}
-				else
-				{
-					UE_LOG(LogWmfMedia, Verbose, TEXT("Session %llx: Failed to get full topology: %s"), this, *WmfMedia::ResultToString(Result));
-
-					if (CurrentState == EMediaState::Preparing)
-					{
-						// incorrect topology for new media
-						CurrentState = EMediaState::Error;
-						DeferredEvents.Enqueue(EMediaEvent::MediaOpenFailed);
-					}
-
-					DiscardPendingState();
-				}
-
-				LastTime = FTimespan::Zero();
-			}
-		}
-		else
-		{
-			if (CurrentState == EMediaState::Preparing)
-			{
-				UE_LOG(LogWmfMedia, Error, TEXT("An error occured when preparing the topology"));
-
-				// an error occurred in the topology
-				CurrentState = EMediaState::Error;
-				DeferredEvents.Enqueue(EMediaEvent::MediaOpenFailed);
-			}
-
-			DiscardPendingState();
-		}
+		HandleSessionTopologyStatus(EventStatus, *Event);
 		break;
 
 	default:
@@ -787,28 +582,27 @@ STDMETHODIMP FWmfMediaSession::Invoke(IMFAsyncResult* AsyncResult)
 	}
 
 	// request next event
-	if ((EventType != MESessionClosed) && (CurrentState != EMediaState::Error) && (MediaSession != NULL))
+	if ((EventType != MESessionClosed) && (SessionState != EMediaState::Error) && (MediaSession != NULL))
 	{
 		const HRESULT Result = MediaSession->BeginGetEvent(this, NULL);
 
 		if (FAILED(Result))
 		{
-			UE_LOG(LogWmfMedia, Warning, TEXT("Failed to request next session event; aborting playback: %s"), *WmfMedia::ResultToString(Result));
+			UE_LOG(LogWmfMedia, Verbose, TEXT("Session %p: Failed to request next session event; aborting playback: %s"), this, *WmfMedia::ResultToString(Result));
 
 			Capabilities = 0;
-			CurrentState = EMediaState::Error;
+			SessionState = EMediaState::Error;
 		}
 	}
 
-	// log current state
-	UE_LOG(LogWmfMedia, VeryVerbose, TEXT("Session %llx: CurrentState: %s, CurrentRate: %f, RequestedRate: %f, CurrentTime: %s, RequestedTime: %s, Pending Change: %i"),
+	UE_LOG(LogWmfMedia, VeryVerbose, TEXT("Session %p: CurrentState: %s, CurrentRate: %f, CurrentTime: %s, SessionState: %s, SessionRate: %f, PendingChanges: %s"),
 		this,
-		WmfMediaSession::StateToString(CurrentState),
-		CurrentRate,
-		RequestedRate,
+		WmfMediaSession::StateToString(GetState()),
+		GetRate(),
 		*GetTime().ToString(),
-		*RequestedTime.ToString(),
-		PendingChange
+		WmfMediaSession::StateToString(SessionState),
+		SessionRate,
+		PendingChanges ? TEXT("true") : TEXT("false")
 	);
 	
 	return S_OK;
@@ -852,120 +646,201 @@ STDMETHODIMP_(ULONG) FWmfMediaSession::Release()
 /* FWmfMediaSession implementation
  *****************************************************************************/
 
-bool FWmfMediaSession::CommitRate(float Rate, EMediaRateThinning Thinning)
+bool FWmfMediaSession::CommitRate(float Rate)
 {
 	check(MediaSession != NULL);
-	check(PendingChange == false);
+	check(!PendingChanges);
 
-	// pause player for basic pausing or if scrubbing isn't supported
-	const bool ShouldPause = (Rate == 0.0f) && ((CurrentRate == 1.0f) || !CanControl(EMediaControl::Scrub));
+	UE_LOG(LogWmfMedia, Verbose, TEXT("Session %p: Committing rate %f"), this, Rate);
 
-	if (ShouldPause)
+	// Note: if rate control is not available, the session only supports pause and play
+
+	if (RateControl == NULL)
 	{
-		if (CurrentState == EMediaState::Playing)
+		if (Rate == 0.0f)
 		{
-			const HRESULT Result = MediaSession->Pause();
-
-			if (FAILED(Result))
+			if (SessionState == EMediaState::Playing)
 			{
-				UE_LOG(LogWmfMedia, Warning, TEXT("Failed to pause for rate change"));
-				return false;
+				const HRESULT Result = MediaSession->Pause();
+
+				if (FAILED(Result))
+				{
+					UE_LOG(LogWmfMedia, Verbose, TEXT("Session %p: Failed to pause session: %s"), this, *WmfMedia::ResultToString(Result));
+					return false;
+				}
 			}
+		}
+		else
+		{
+			if (SessionState != EMediaState::Playing)
+			{
+				PROPVARIANT StartPosition;
+				PropVariantInit(&StartPosition);
 
-			UE_LOG(LogWmfMedia, Verbose, TEXT("Session %llx: Paused for rate change from %f to %f"), this, CurrentRate, Rate);
+				const HRESULT Result = MediaSession->Start(NULL, &StartPosition);
 
-			// request deferred rate change
-			if (RateControl != NULL)
-			{
-				RequestedRate = Rate;
-				RequestedThinning = Thinning;
+				if (FAILED(Result))
+				{
+					UE_LOG(LogWmfMedia, Verbose, TEXT("Session %p: Failed to start session: %s"), this, *WmfMedia::ResultToString(Result));
+					return false;
+				}
 			}
+		}
 
-			PendingChange = true;
+		PendingChanges = true;
 
-			return true;
-		}
+		return true;
 	}
 
-	// most rate changes require transition to Stopped first
-	if (((Rate >= 0.0f) || (CurrentRate >= 0.0f)) && ((Rate <= 0.0f) || (CurrentRate <= 0.0f)))
+	// Note: if rate control is available, things get considerably more complicated
+	// as many rate transitions are only allowed from certain session states
+
+	if (((Rate >= 0.0f) && (SessionRate < 0.0f)) || ((Rate < 0.0f) && (SessionRate >= 0.0f)))
 	{
-		if (!ShouldPause && (CurrentState != EMediaState::Stopped))
+		// transitions between negative and zero/positive rates require Stopped state
+		if (SessionState != EMediaState::Stopped)
 		{
+			UE_LOG(LogWmfMedia, Verbose, TEXT("Session %p: Stopping session for rate change"), this);
+
 			LastTime = GetTime();
+
+			// stop the media session
 			const HRESULT Result = MediaSession->Stop();
 
 			if (FAILED(Result))
 			{
-				UE_LOG(LogWmfMedia, Warning, TEXT("Failed to stop for rate change"));
+				UE_LOG(LogWmfMedia, Verbose, TEXT("Session %p: Failed to stop for rate change: %s"), this, *WmfMedia::ResultToString(Result));
 				return false;
 			}
 
-			UE_LOG(LogWmfMedia, Verbose, TEXT("Session %llx: Stopped for rate change from %f to %f"), this, CurrentRate, Rate);
+			// defer rate change
+			UE_LOG(LogWmfMedia, Verbose, TEXT("Session %p: Deferring rate change until after pending stop"), this);
 
-			// request deferred rate change
-			if (RequestedTime == FTimespan::MinValue())
-			{
+			if (!RequestedTime.IsSet())
+			{			
 				RequestedTime = LastTime;
 			}
 
 			RequestedRate = Rate;
-			RequestedThinning = Thinning;
-			PendingChange = true;
+			PendingChanges = true;
 
 			return true;
 		}
 	}
+	
+	if (((Rate == 0.0f) && (SessionRate != 0.0f)) || ((Rate != 0.0f) && (SessionRate == 0.0f)))
+	{
+		// transitions between positive and zero rates require Paused or Stopped state
+		if ((SessionState != EMediaState::Paused) && (SessionState != EMediaState::Stopped))
+		{
+			UE_LOG(LogWmfMedia, Verbose, TEXT("Session %p: Pausing session for rate change from %f to %f"), this, SessionRate, Rate);
+
+			// pause the media session
+			const HRESULT Result = MediaSession->Pause();
 
-	// determine restart time
-	FTimespan RestartTime;
+			if (FAILED(Result))
+			{
+				UE_LOG(LogWmfMedia, Verbose, TEXT("Session %p: Failed to pause for rate change: %s"), this, *WmfMedia::ResultToString(Result));
+				return false;
+			}
 
-	if (RequestedTime >= FTimespan::Zero())
-	{
-		RestartTime = RequestedTime;
-	}
-	else
-	{
-		RestartTime = GetTime();
-	}
+			// defer rate change
+			UE_LOG(LogWmfMedia, Verbose, TEXT("Session %p: Deferring rate change until after pending pause"), this);
 
-	if ((RestartTime == FTimespan::Zero()) && (Rate < 0.0f))
-	{
-		RestartTime = CurrentDuration; // loop to end
-	}
-	else if ((RestartTime == CurrentDuration) && (Rate > 0.0f))
-	{
-		RestartTime = FTimespan::Zero(); // loop to beginning
+			RequestedRate = Rate;
+			PendingChanges = true;
+
+			return true;
+		}
 	}
 
-	// set desired rate
-	if (RateControl != NULL)
+	// Note: The rate control could be updated right after requesting the Pause or Stopped
+	// states above, but we wait for these transitions to complete, so that multiple calls
+	// to SetRate do not interfere with each other.
+
+	if (Rate != SessionRate)
 	{
+		// determine thinning mode
+		EMediaRateThinning Thinning;
+
+		if (UnthinnedRates.Contains(Rate))
+		{
+			Thinning = EMediaRateThinning::Unthinned;
+		}
+		else if (ThinnedRates.Contains(Rate))
+		{
+			Thinning = EMediaRateThinning::Thinned;
+		}
+		else
+		{
+			return false;
+		}
+
 		const TCHAR* ThinnedString = (Thinning == EMediaRateThinning::Thinned) ? TEXT("thinned") : TEXT("unthinned");
+		UE_LOG(LogWmfMedia, Verbose, TEXT("Session %p: Changing rate from %f to %f [%s]"), this, SessionRate, Rate, ThinnedString);
+
+		// set the new rate
 		const HRESULT Result = RateControl->SetRate((Thinning == EMediaRateThinning::Thinned) ? TRUE : FALSE, Rate);
 
 		if (FAILED(Result))
 		{
-			UE_LOG(LogWmfMedia, Warning, TEXT("Failed to commit rate change from %f to %f (%s): %s"), CurrentRate, Rate, ThinnedString, *WmfMedia::ResultToString(Result));
+			UE_LOG(LogWmfMedia, Verbose, TEXT("Session %p: Failed to change rate: %s"), this, *WmfMedia::ResultToString(Result));
 			return false;
 		}
 
-		UE_LOG(LogWmfMedia, Verbose, TEXT("Session %llx: Committed rate change from %f to %f (%s)"), this, CurrentRate, Rate, ThinnedString);
-
-		// request deferred restart
-		if (Rate != 0.0f)
+		if (PendingChanges)
 		{
-			RequestedTime = RestartTime;
-			PendingChange = true;
+			return true; // wait for required state transitions to complete
 		}
 
+		PendingChanges = true;
+	}
+
+	// Note: no further changes needed if session was playing and direction didn't change
+
+	if (((Rate * SessionRate) > 0.0f) && (SessionState == EMediaState::Playing))
+	{
 		return true;
 	}
 
-	// restart playback
-	if (Rate != 0.0f)
+	// Note: for non-zero rates, the session must be restarted. If the rate control wasn't
+	// updated above, this can be done immediately, otherwise it has to be deferred until
+	// after the rate control finished setting the new rate.
+
+	if ((Rate != 0.0f) && ((SessionState != EMediaState::Playing) || PendingChanges))
 	{
-		CommitTime(RestartTime);
+		// determine restart time
+		FTimespan RestartTime;
+
+		if (RequestedTime.IsSet())
+		{
+			RestartTime = RequestedTime.GetValue();
+		}
+		else
+		{
+			RestartTime = GetTime();
+		}
+
+		if ((RestartTime == FTimespan::Zero()) && (Rate < 0.0f))
+		{
+			RestartTime = CurrentDuration; // loop to end
+		}
+		else if ((RestartTime == CurrentDuration) && (Rate > 0.0f))
+		{
+			RestartTime = FTimespan::Zero(); // loop to beginning
+		}
+
+		// restart session
+		if (PendingChanges)
+		{
+			UE_LOG(LogWmfMedia, Verbose, TEXT("Session %p: Requesting start after pending rate change"), this);
+			RequestedTime = RestartTime;
+		}
+		else
+		{
+			UE_LOG(LogWmfMedia, Verbose, TEXT("Session %p: Starting session for rate change"), this);
+			CommitTime(RestartTime);
+		}
 	}
 
 	return true;
@@ -975,28 +850,21 @@ bool FWmfMediaSession::CommitRate(float Rate, EMediaRateThinning Thinning)
 bool FWmfMediaSession::CommitTime(FTimespan Time)
 {
 	check(MediaSession != NULL);
-	check(PendingChange == false);
+	check(PendingChanges == false);
 
-	// reapply topology if needed
-#if 0
-	if ((CurrentState != EMediaState::Paused) && (CurrentState != EMediaState::Playing))
-	{
-		HRESULT Result = MediaSession->SetTopology(0, CurrentTopology);
+	const FString TimeString = (Time == WmfMediaSession::RequestedTimeCurrent) ? TEXT("<current>") : *Time.ToString();
+	UE_LOG(LogWmfMedia, Verbose, TEXT("Session %p: Committing time %s"), this, *TimeString);
 
-		if (FAILED(Result))
-		{
-			UE_LOG(LogWmfMedia, Error, TEXT("Failed to reapply topology (%s)"), *WmfMedia::ResultToString(Result));
-			return false;
-		}
+	// start session at requested time
+	PROPVARIANT StartPosition;
 
-		UE_LOG(LogWmfMedia, Verbose, TEXT("Session %llx: Reapplied topology"), this);
+	if (!CanControl(EMediaControl::Seek))
+	{
+		UE_LOG(LogWmfMedia, Verbose, TEXT("Session %p: Starting from <current>, because media can't seek"), this);
+		Time = WmfMediaSession::RequestedTimeCurrent;
 	}
-#endif
-
-	// seek to requested time
-	PROPVARIANT StartPosition;
 
-	if (!CanControl(EMediaControl::Seek) || (Time == FTimespan::MaxValue()))
+	if (Time == WmfMediaSession::RequestedTimeCurrent)
 	{
 		StartPosition.vt = VT_EMPTY; // current time
 	}
@@ -1010,13 +878,11 @@ bool FWmfMediaSession::CommitTime(FTimespan Time)
 
 	if (FAILED(Result))
 	{
-		UE_LOG(LogWmfMedia, Warning, TEXT("Failed to commit time %s: %s"), (Time == FTimespan::MaxValue()) ? TEXT("current") : *Time.ToString(), *WmfMedia::ResultToString(Result));
+		UE_LOG(LogWmfMedia, Verbose, TEXT("Session %p: Failed to start session: %s"), this, *WmfMedia::ResultToString(Result));
 		return false;
 	}
 
-	UE_LOG(LogWmfMedia, Verbose, TEXT("Session %llx: Commited time %s"), this, (Time == FTimespan::MaxValue()) ? TEXT("current") : *Time.ToString());
-
-	PendingChange = true;
+	PendingChanges = true;
 
 	return true;
 }
@@ -1025,10 +891,14 @@ bool FWmfMediaSession::CommitTime(FTimespan Time)
 bool FWmfMediaSession::CommitTopology(IMFTopology* Topology)
 {
 	check(MediaSession != NULL);
-	check(PendingChange == false);
+	check(PendingChanges == false);
 
-	if (CurrentState != EMediaState::Stopped)
+	UE_LOG(LogWmfMedia, Verbose, TEXT("Session %p: Committing topology %p"), this, Topology);
+
+	if (SessionState != EMediaState::Stopped)
 	{
+		UE_LOG(LogWmfMedia, Verbose, TEXT("Session %p: Stopping session for topology change"), this);
+
 		LastTime = GetTime();
 
 		// topology change requires transition to Stopped; playback is resumed afterwards
@@ -1036,92 +906,519 @@ bool FWmfMediaSession::CommitTopology(IMFTopology* Topology)
 
 		if (FAILED(Result))
 		{
-			UE_LOG(LogWmfMedia, Warning, TEXT("Failed to stop for topology change"));
+			UE_LOG(LogWmfMedia, Verbose, TEXT("Session %p: Failed to stop for topology change: %s"), *WmfMedia::ResultToString(Result));
 			return false;
 		}
 
-		UE_LOG(LogWmfMedia, Verbose, TEXT("Session %llx: Stopped for topology change"), this);
-
 		// request deferred playback restart
-		if (RequestedTime == FTimespan::MinValue())
+		if (!RequestedTime.IsSet())
 		{
+			UE_LOG(LogWmfMedia, Verbose, TEXT("Session %p: Requesting restart after pending stop"), this);
+
 			RequestedTime = LastTime;
 		}
 
+		UE_LOG(LogWmfMedia, Verbose, TEXT("Session %p: Requesting topology change after pending stop"), this);
+
 		RequestedTopology = Topology;
-		PendingChange = true;
+		PendingChanges = true;
 
 		return true;
 	}
 
 	// set new topology
-	HRESULT Result = MediaSession->SetTopology(MFSESSION_SETTOPOLOGY_IMMEDIATE, Topology);
+	HRESULT Result = MediaSession->ClearTopologies();
 
-	if (SUCCEEDED(Result))
+	if (FAILED(Result))
 	{
-		UE_LOG(LogWmfMedia, Verbose, TEXT("Session %llx: Committed topology"), this);
+		UE_LOG(LogWmfMedia, Verbose, TEXT("Session %p: Failed to clear queued topologies: %s"), this, *WmfMedia::ResultToString(Result));
+		return false;
 	}
-	else
+
+	if (Topology != NULL)
 	{
-		UE_LOG(LogWmfMedia, Warning, TEXT("Failed to set topology: %s"), *WmfMedia::ResultToString(Result));
-	}
+		Result = MediaSession->SetTopology(MFSESSION_SETTOPOLOGY_IMMEDIATE, Topology);
+
+		if (FAILED(Result))
+		{
+			UE_LOG(LogWmfMedia, Verbose, TEXT("Session %p: Failed to set topology: %s"), this, *WmfMedia::ResultToString(Result));
+			return false;
+		}
 
-	PendingChange = true;
+		UE_LOG(LogWmfMedia, Verbose, TEXT("Session %p: Committed topology %p"), this, Topology);
+
+		PendingChanges = true;
+	}
 
 	return true;
 }
 
 
-void FWmfMediaSession::DiscardPendingState()
+void FWmfMediaSession::DiscardPendingChanges()
 {
-	PendingChange = false;
-	RequestedRate = CurrentRate;
-	RequestedTime = FTimespan::MinValue();
+	UE_LOG(LogWmfMedia, Verbose, TEXT("Session %p: Discarding pending changes"), this);
+
+	RequestedRate.Reset();
+	RequestedTime.Reset();
 	RequestedTopology.Reset();
+
+	PendingChanges = false;
 }
 
 
-void FWmfMediaSession::UpdatePendingState()
+void FWmfMediaSession::DoPendingChanges()
 {
-	UE_LOG(LogWmfMedia, Verbose, TEXT("Session %llx: Updating pending state"), this);
+	if (PendingChanges)
+	{
+		const FString RequestedRateString = !RequestedRate.IsSet() ? TEXT("<none>") : FString::Printf(TEXT("%f"), RequestedRate.GetValue());
+		const FString RequestedTimeString = !RequestedTime.IsSet() ? TEXT("<none>") : (RequestedTime.GetValue() == WmfMediaSession::RequestedTimeCurrent) ? TEXT("<current>") : *RequestedTime.GetValue().ToString();
+		const FString RequestedTopologyString = !RequestedTopology ? TEXT("<none>") : FString::Printf(TEXT("%p"), RequestedTopology.Get());
 
-	PendingChange = false;
+		UE_LOG(LogWmfMedia, Verbose, TEXT("Session %p: Doing pending changes: RequestedRate: %s, RequestedTime: %s, RequestedTopology: %s"), this, *RequestedRateString, *RequestedTimeString, *RequestedTopologyString);
+	}
+	else
+	{
+		UE_LOG(LogWmfMedia, Verbose, TEXT("Session %p: Doing pending changes: none"), this);
+	}
+
+	PendingChanges = false;
 
 	// commit pending topology changes
 	if (RequestedTopology != NULL)
 	{
-		if (!CommitTopology(RequestedTopology))
+		TComPtr<IMFTopology> Topology = RequestedTopology;
+		RequestedTopology.Reset();
+
+		CommitTopology(Topology);
+
+		if (PendingChanges)
+		{
+			return;
+		}
+	}
+
+	// commit pending rate changes
+	if (RequestedRate.IsSet())
+	{
+		const float Rate = RequestedRate.GetValue();
+		RequestedRate.Reset();
+
+		CommitRate(Rate);
+
+		if (PendingChanges)
+		{
+			return;
+		}
+	}
+
+	// commit pending seeks/restarts
+	if (RequestedTime.IsSet())
+	{
+		const FTimespan Time = RequestedTime.GetValue();
+		RequestedTime.Reset();
+
+		CommitTime(Time);
+	}
+}
+
+
+void FWmfMediaSession::UpdateCharacteristics()
+{
+	// reset characteristics
+	PresentationClock.Reset();
+	RateControl.Reset();
+	RateSupport.Reset();
+
+	ThinnedRates.Empty();
+	UnthinnedRates.Empty();
+
+	CanScrub = false;
+
+	if (MediaSession == NULL)
+	{
+		return;
+	}
+
+	// get presentation clock, if available
+	TComPtr<IMFClock> Clock;
+
+	HRESULT Result = MediaSession->GetClock(&Clock);
+
+	if (FAILED(Result))
+	{
+		UE_LOG(LogWmfMedia, Verbose, TEXT("Session %p: Session clock unavailable: %s"), this, *WmfMedia::ResultToString(Result));
+	}
+	else
+	{
+		Result = Clock->QueryInterface(IID_PPV_ARGS(&PresentationClock));
+
+		if (FAILED(Result))
+		{
+			UE_LOG(LogWmfMedia, Verbose, TEXT("Session %p: Presentation clock unavailable: %s"), this, *WmfMedia::ResultToString(Result));
+		}
+		else
+		{
+			UE_LOG(LogWmfMedia, Verbose, TEXT("Session %p: Presentation clock ready"), this);
+		}
+	}
+
+	// get rate control & rate support, if available
+	Result = ::MFGetService(MediaSession, MF_RATE_CONTROL_SERVICE, IID_PPV_ARGS(&RateControl));
+
+	if (FAILED(Result))
+	{
+		UE_LOG(LogWmfMedia, Verbose, TEXT("Session %p: Rate control service unavailable: %s"), this, *WmfMedia::ResultToString(Result));
+	}
+	else
+	{
+		UE_LOG(LogWmfMedia, Verbose, TEXT("Session %p: Rate control ready"), this);
+
+		if (FAILED(RateControl->GetRate(FALSE, &SessionRate)))
+		{
+			UE_LOG(LogWmfMedia, Verbose, TEXT("Session %p: Failed to initialize current rate"), this);
+			SessionRate = 1.0f; // the session's initial play rate is usually 1.0
+		}
+	}
+
+	Result = ::MFGetService(MediaSession, MF_RATE_CONTROL_SERVICE, IID_PPV_ARGS(&RateSupport));
+
+	if (FAILED(Result))
+	{
+		UE_LOG(LogWmfMedia, Verbose, TEXT("Session %p: Rate support service unavailable: %s"), this, *WmfMedia::ResultToString(Result));
+	}
+	else
+	{
+		UE_LOG(LogWmfMedia, Verbose, TEXT("Session %p: Rate support ready"), this);
+	}
+
+	// cache rate control properties
+	if (RateSupport.IsValid())
+	{
+		CanScrub = SUCCEEDED(RateSupport->IsRateSupported(TRUE, 0.0f, NULL));
+
+		float MaxRate = 0.0f;
+		float MinRate = 0.0f;
+
+		if (SUCCEEDED(RateSupport->GetSlowestRate(MFRATE_FORWARD, TRUE, &MinRate)) &&
+			SUCCEEDED(RateSupport->GetFastestRate(MFRATE_FORWARD, TRUE, &MaxRate)))
+		{
+			ThinnedRates.Add(TRange<float>::Inclusive(MinRate, MaxRate));
+		}
+
+		if (SUCCEEDED(RateSupport->GetSlowestRate(MFRATE_REVERSE, TRUE, &MinRate)) &&
+			SUCCEEDED(RateSupport->GetFastestRate(MFRATE_REVERSE, TRUE, &MaxRate)))
+		{
+			ThinnedRates.Add(TRange<float>::Inclusive(MaxRate, MinRate));
+		}
+
+		if (SUCCEEDED(RateSupport->GetSlowestRate(MFRATE_FORWARD, FALSE, &MinRate)) &&
+			SUCCEEDED(RateSupport->GetFastestRate(MFRATE_FORWARD, FALSE, &MaxRate)))
+		{
+			UnthinnedRates.Add(TRange<float>::Inclusive(MinRate, MaxRate));
+		}
+
+		if (SUCCEEDED(RateSupport->GetSlowestRate(MFRATE_REVERSE, FALSE, &MinRate)) &&
+			SUCCEEDED(RateSupport->GetFastestRate(MFRATE_REVERSE, FALSE, &MaxRate)))
+		{
+			UnthinnedRates.Add(TRange<float>::Inclusive(MaxRate, MinRate));
+		}
+	}
+}
+
+
+/* FWmfMediaSession callbacks
+*****************************************************************************/
+
+void FWmfMediaSession::HandleError(HRESULT EventStatus)
+{
+	UE_LOG(LogWmfMedia, Error, TEXT("An error occurred in the media session: %s"), *WmfMedia::ResultToString(EventStatus));
+
+	SessionState = EMediaState::Error;
+	DiscardPendingChanges();
+	MediaSession->Close();
+}
+
+
+void FWmfMediaSession::HandleSessionEnded()
+{
+	DeferredEvents.Enqueue(EMediaEvent::PlaybackEndReached);
+
+	SessionState = EMediaState::Stopped;
+
+	if (ShouldLoop)
+	{
+		// loop back to beginning/end
+		RequestedTime = (SessionRate < 0.0f) ? CurrentDuration : FTimespan::Zero();
+		DoPendingChanges();
+	}
+	else
+	{
+		LastTime = FTimespan::Zero();
+		RequestedRate.Reset();
+	}
+}
+
+
+void FWmfMediaSession::HandleSessionPaused(HRESULT EventStatus)
+{
+	if (SUCCEEDED(EventStatus))
+	{
+		SessionState = EMediaState::Paused;
+		DeferredEvents.Enqueue(EMediaEvent::PlaybackSuspended);
+
+		DoPendingChanges();
+	}
+	else
+	{
+		DiscardPendingChanges();
+	}
+}
+
+
+void FWmfMediaSession::HandleSessionRateChanged(HRESULT EventStatus, IMFMediaEvent& Event)
+{
+	if (SUCCEEDED(EventStatus))
+	{
+		// cache current rate
+		PROPVARIANT Value;
+		::PropVariantInit(&Value);
+
+		const HRESULT Result = Event.GetValue(&Value);
+
+		if (SUCCEEDED(Result) && (Value.vt == VT_R4))
+		{
+			SessionRate = Value.fltVal;
+		}
+	}
+	else if (RateControl != NULL)
+	{
+		BOOL Thin = FALSE;
+		RateControl->GetRate(&Thin, &SessionRate);
+	}
+
+	DoPendingChanges();
+}
+
+
+void FWmfMediaSession::HandleSessionScrubSampleComplete()
+{
+//	DeferredEvents.Enqueue(EMediaEvent::SeekCompleted);
+//	DoPendingChanges();
+}
+
+
+void FWmfMediaSession::HandleSessionStarted(HRESULT EventStatus)
+{
+	if (SUCCEEDED(EventStatus))
+	{
+		if ((SessionState == EMediaState::Paused) && (SessionRate == 0.0f))
 		{
-			RequestedTopology.Reset();
+			SessionState = EMediaState::Playing;
+
+			// Note: pending changes will be processed in MESessionScrubSampleComplete
+
+			DeferredEvents.Enqueue(EMediaEvent::SeekCompleted);
+			DoPendingChanges();
+		}
+		else
+		{
+			if (RateControl == NULL)
+			{
+				SessionRate = 1.0f;
+			}
+			else
+			{
+				BOOL Thin = FALSE;
+				RateControl->GetRate(&Thin, &SessionRate);
+			}
+
+			if (SessionState == EMediaState::Playing)
+			{
+				DeferredEvents.Enqueue(EMediaEvent::SeekCompleted);
+			}
+			else
+			{
+				SessionState = EMediaState::Playing;
+
+				if (SessionRate == 0.0f)
+				{
+					RequestedRate = 0.0f;
+					PendingChanges = true;
+				}
+				else
+				{
+					DeferredEvents.Enqueue(EMediaEvent::PlaybackResumed);
+				}
+			}
+			
+			DoPendingChanges();
 		}
+	}
+	else
+	{
+		DiscardPendingChanges();
+	}
+}
 
-		if (PendingChange)
+
+void FWmfMediaSession::HandleSessionStopped(HRESULT EventStatus)
+{
+	if (SUCCEEDED(EventStatus))
+	{
+		SessionState = EMediaState::Stopped;
+		DeferredEvents.Enqueue(EMediaEvent::PlaybackSuspended);
+		DoPendingChanges();
+	}
+	else
+	{
+		DiscardPendingChanges();
+	}
+}
+
+
+void FWmfMediaSession::HandleSessionTopologySet(HRESULT EventStatus, IMFMediaEvent& Event)
+{
+	if (SUCCEEDED(EventStatus))
+	{
+		const HRESULT Result = WmfMedia::GetTopologyFromEvent(Event, CurrentTopology);
+
+		if (SUCCEEDED(Result))
 		{
+			UE_LOG(LogWmfMedia, Verbose, TEXT("Session %p: Topology %p set as current"), this, CurrentTopology.Get());
+
+			if (SessionState != EMediaState::Preparing)
+			{
+				// Note: track and format changes won't send an MF_TOPOSTATUS_READY event
+				// until playback is restarted, so we do pending changes here instead.
+
+				DoPendingChanges();
+			}
+
 			return;
 		}
+
+		UE_LOG(LogWmfMedia, Verbose, TEXT("Session %p: Failed to get topology that was set: %s"), this, *WmfMedia::ResultToString(Result));
 	}
 
-	// commit pending rate changes
-	if (RequestedRate != CurrentRate)
+	if (SessionState == EMediaState::Preparing)
+	{
+		SessionState = EMediaState::Error;
+		DeferredEvents.Enqueue(EMediaEvent::MediaOpenFailed);
+	}
+
+	DiscardPendingChanges();
+}
+
+
+void FWmfMediaSession::HandleSessionTopologyStatus(HRESULT EventStatus, IMFMediaEvent& Event)
+{
+	// get the status of the topology that generated the event
+	MF_TOPOSTATUS TopologyStatus = MF_TOPOSTATUS_INVALID;
 	{
-		if (!CommitRate(RequestedRate, RequestedThinning))
+		const HRESULT Result = Event.GetUINT32(MF_EVENT_TOPOLOGY_STATUS, (UINT32*)&TopologyStatus);
+
+		if (FAILED(Result))
 		{
-			RequestedRate = CurrentRate;
+			UE_LOG(LogWmfMedia, Verbose, TEXT("Session %p: Failed to get topology status: %s"), this, *WmfMedia::ResultToString(Result));
+			return;
 		}
+	}
 
-		if (PendingChange)
+	// get the topology that generated the event
+	TComPtr<IMFTopology> Topology;
+	{
+		const HRESULT Result = WmfMedia::GetTopologyFromEvent(Event, Topology);
+
+		if (FAILED(Result))
 		{
+			UE_LOG(LogWmfMedia, Verbose, TEXT("Session %p: Failed to get topology from topology status event: %s"), this, *WmfMedia::ResultToString(Result));
 			return;
 		}
 	}
 
-	// commit pending seeks/restarts
-	if (RequestedTime >= FTimespan::Zero())
+	UE_LOG(LogWmfMedia, Verbose, TEXT("Session %p: Topology %p changed status to %s"), this, Topology.Get(), *WmfMedia::TopologyStatusToString(TopologyStatus));
+
+	// Note: the ordering of topology status events is not guaranteed for two
+	// consecutive topologies. we skip events that are not the for current one.
+
+	if (Topology != CurrentTopology)
+	{
+		return;
+	}
+
+	if (SessionState == EMediaState::Error)
 	{
-		CommitTime(RequestedTime);
+		DiscardPendingChanges();
 
-		RequestedTime = FTimespan::MinValue();
+		return;
 	}
+
+	if (FAILED(EventStatus))
+	{
+		if (SessionState == EMediaState::Preparing)
+		{
+			UE_LOG(LogWmfMedia, Error, TEXT("An error occured when preparing the topology"), this);
+
+			// an error occurred in the topology
+			SessionState = EMediaState::Error;
+			DeferredEvents.Enqueue(EMediaEvent::MediaOpenFailed);
+		}
+
+		DiscardPendingChanges();
+
+		return;
+	}
+
+	if (TopologyStatus != MF_TOPOSTATUS_READY)
+	{
+		return;
+	}
+
+#if 0
+	TComPtr<IMFTopology> FullTopology;
+	{
+		const HRESULT Result = MediaSession->GetFullTopology(MFSESSION_GETFULLTOPOLOGY_CURRENT, 0, &FullTopology);
+
+		if (SUCCEEDED(Result))
+		{
+			if (FullTopology != CurrentTopology)
+			{
+				UE_LOG(LogWmfMedia, Verbose, TEXT("Session %p: Full topology %p from media session doesn't match current topology %p"), this, FullTopology.Get(), CurrentTopology.Get());
+			}
+		}
+		else
+		{
+			UE_LOG(LogWmfMedia, Verbose, TEXT("Session %p: Failed to get full topology from media session"), this, *WmfMedia::ResultToString(Result));
+		}
+	}
+#endif
+
+	// initialize new topology
+	UpdateCharacteristics();
+
+	// new media opened successfully
+	if (SessionState == EMediaState::Preparing)
+	{
+		UE_LOG(LogWmfMedia, Verbose, TEXT("Session %p: Topology %p ready"), this, CurrentTopology.Get());
+
+		SessionState = EMediaState::Stopped;
+		DeferredEvents.Enqueue(EMediaEvent::MediaOpened);
+	}
+	else if (SessionState == EMediaState::Paused)
+	{
+		// Note: when paused, the new topology won't apply until the next session start,
+		// so we request a scrub to the current time in order to update the video frame.
+
+		if (!RequestedTime.IsSet())
+		{
+			UE_LOG(LogWmfMedia, Verbose, TEXT("Session %p: Requesting scrub after topology change"), this);
+
+			RequestedTime = WmfMediaSession::RequestedTimeCurrent;
+			PendingChanges = true;
+		}
+	}
+
+	DoPendingChanges();
 }
 
 
diff --git a/Engine/Plugins/Media/WmfMedia/Source/WmfMedia/Private/Wmf/WmfMediaSession.h b/Engine/Plugins/Media/WmfMedia/Source/WmfMedia/Private/Wmf/WmfMediaSession.h
--- a/Engine/Plugins/Media/WmfMedia/Source/WmfMedia/Private/Wmf/WmfMediaSession.h
+++ b/Engine/Plugins/Media/WmfMedia/Source/WmfMedia/Private/Wmf/WmfMediaSession.h
@@ -11,6 +11,7 @@
 #include "HAL/CriticalSection.h"
 #include "IMediaControls.h"
 #include "IMediaEventSink.h"
+#include "Misc/Optional.h"
 #include "Misc/Timespan.h"
 
 #include "AllowWindowsPlatformTypes.h"
@@ -119,10 +120,9 @@ protected:
 	 * The caller holds the lock to the critical section.
 	 *
 	 * @param Rate The play rate to commit.
-	 * @param Thinning The thinning mode.
 	 * @see CommitTime, CommitTopology
 	 */
-	bool CommitRate(float Rate, EMediaRateThinning Thinning);
+	bool CommitRate(float Rate);
 
 	/**
 	 * Commit the specified play position.
@@ -145,16 +145,48 @@ protected:
 	/**
 	 * Discard all pending state changes.
 	 *
-	 * @see UpdatePendingState
+	 * @see DoPendingChanges
 	 */
-	void DiscardPendingState();
+	void DiscardPendingChanges();
 
 	/**
-	 * Applies any pending state changes.
-	 *
-	 * @see DiscardPendingState
-	 */
-	void UpdatePendingState();
+	* Applies any pending state changes.
+	*
+	* @see DiscardPendingChanges
+	*/
+	void DoPendingChanges();
+
+	/** Get the latest characteristics from the current media source. */
+	void UpdateCharacteristics();
+
+private:
+
+	/** Callback for the MEError event. */
+	void HandleError(HRESULT EventStatus);
+
+	/** Callback for the MESessionEnded event. */
+	void HandleSessionEnded();
+
+	/** Callback for the MESessionPaused event. */
+	void HandleSessionPaused(HRESULT EventStatus);
+
+	/** Callback for the MESessionRateChanged event. */
+	void HandleSessionRateChanged(HRESULT EventStatus, IMFMediaEvent& Event);
+
+	/** Callback for the MESessionScrubSampleComplete event. */
+	void HandleSessionScrubSampleComplete();
+
+	/** Callback for the MESessionStarted event. */
+	void HandleSessionStarted(HRESULT EventStatus);
+
+	/** Callback for the MESessionStopped event. */
+	void HandleSessionStopped(HRESULT EventStatus);
+
+	/** Callback for the MESessionTopologySet event. */
+	void HandleSessionTopologySet(HRESULT EventStatus, IMFMediaEvent& Event);
+
+	/** Callback for the MESessionTopologyStatus event. */
+	void HandleSessionTopologyStatus(HRESULT EventStatus, IMFMediaEvent& Event);
 
 private:
 
@@ -175,12 +207,6 @@ private:
 	/** The duration of the media. */
 	FTimespan CurrentDuration;
 
-	/** The current playback rate. */
-	float CurrentRate;
-
-	/** The current playback state. */
-	EMediaState CurrentState;
-
 	/** The full playback topology currently set on the media session. */
 	TComPtr<IMFTopology> CurrentTopology;
 
@@ -193,8 +219,8 @@ private:
 	/** The last play head position before playback was stopped. */
 	FTimespan LastTime;
 
-	/** Whether a state change is pending. */
-	bool PendingChange;
+	/** Whether one or more state changes are pending. */
+	bool PendingChanges;
 
 	/** The media session's clock. */
 	TComPtr<IMFPresentationClock> PresentationClock;
@@ -209,16 +235,19 @@ private:
 	int32 RefCount;
 
 	/** Deferred play rate change value. */
-	float RequestedRate;
-
-	/** Thinning of the deferred play rate. */
-	EMediaRateThinning RequestedThinning;
+	TOptional<float> RequestedRate;
 
 	/** Deferred playback topology to set. */
 	TComPtr<IMFTopology> RequestedTopology;
 
-	/** Deferred play time change value. */
-	FTimespan RequestedTime;
+	/** Deferred play time change value (MinValue = no change, MaxValue = current time). */
+	TOptional<FTimespan> RequestedTime;
+
+	/** The session's internal playback rate (not necessarily the same as GetRate). */
+	float SessionRate;
+
+	/** The session's current state (not necessarily the same as GetState). */
+	EMediaState SessionState;
 
 	/** Whether playback should loop to the beginning. */
 	bool ShouldLoop;
diff --git a/Engine/Plugins/Media/WmfMedia/Source/WmfMedia/Private/Wmf/WmfMediaSink.cpp b/Engine/Plugins/Media/WmfMedia/Source/WmfMedia/Private/Wmf/WmfMediaSink.cpp
--- a/Engine/Plugins/Media/WmfMedia/Source/WmfMedia/Private/Wmf/WmfMediaSink.cpp
+++ b/Engine/Plugins/Media/WmfMedia/Source/WmfMedia/Private/Wmf/WmfMediaSink.cpp
@@ -32,6 +32,8 @@ bool FWmfMediaSink::Initialize(FWmfMediaStreamSink& InStreamSink)
 {
 	FScopeLock Lock(&CriticalSection);
 
+	DWORD StreamId = 0;
+
 	if (!InStreamSink.Initialize(*this))
 	{
 		return false;
@@ -55,7 +57,7 @@ STDMETHODIMP FWmfMediaSink::OnClockPause(MFTIME hnsSystemTime)
 		return MF_E_SHUTDOWN;
 	}
 
-	return S_OK;
+	return StreamSink->Pause();
 }
 
 
@@ -68,7 +70,7 @@ STDMETHODIMP FWmfMediaSink::OnClockRestart(MFTIME hnsSystemTime)
 		return MF_E_SHUTDOWN;
 	}
 
-	return S_OK;
+	return StreamSink->Restart();
 }
 
 
@@ -94,7 +96,7 @@ STDMETHODIMP FWmfMediaSink::OnClockStart(MFTIME hnsSystemTime, LONGLONG llClockS
 		return MF_E_SHUTDOWN;
 	}
 
-	return S_OK;
+	return StreamSink->Start();
 }
 
 
@@ -107,7 +109,26 @@ STDMETHODIMP FWmfMediaSink::OnClockStop(MFTIME hnsSystemTime)
 		return MF_E_SHUTDOWN;
 	}
 
-	return S_OK;
+	return StreamSink->Stop();
+}
+
+
+/* IMFGetService interface
+ *****************************************************************************/
+
+STDMETHODIMP FWmfMediaSink::GetService(__RPC__in REFGUID guidService, __RPC__in REFIID riid, __RPC__deref_out_opt LPVOID* ppvObject)
+{
+	if (guidService == MF_RATE_CONTROL_SERVICE)
+	{
+		return QueryInterface(riid, ppvObject);
+	}
+	
+	if (guidService == MR_VIDEO_ACCELERATION_SERVICE)
+	{
+		// @todo gmp: implement WmfMedia video acceleration
+	}
+
+	return MF_E_UNSUPPORTED_SERVICE;
 }
 
 
@@ -173,11 +194,6 @@ STDMETHODIMP FWmfMediaSink::GetStreamSinkById(DWORD dwIdentifier, __RPC__deref_o
 		return E_POINTER;
 	}
 
-	if (dwIdentifier != WmfMediaStreamSink::FixedStreamId)
-	{
-		return MF_E_INVALIDSTREAMNUMBER;
-	}
-
 	FScopeLock Lock(&CriticalSection);
 
 	if (!StreamSink.IsValid())
@@ -185,6 +201,13 @@ STDMETHODIMP FWmfMediaSink::GetStreamSinkById(DWORD dwIdentifier, __RPC__deref_o
 		return MF_E_SHUTDOWN;
 	}
 
+	DWORD Identifier = 0;
+
+	if (FAILED(StreamSink->GetIdentifier(&Identifier)) || (Identifier != dwIdentifier))
+	{
+		return MF_E_INVALIDSTREAMNUMBER;
+	}
+
 	*ppStreamSink = StreamSink;
 	(*ppStreamSink)->AddRef();
 
@@ -260,7 +283,7 @@ STDMETHODIMP FWmfMediaSink::SetPresentationClock(__RPC__in_opt IMFPresentationCl
 
 		if (FAILED(Result))
 		{
-			UE_LOG(LogWmfMedia, Error, TEXT("Failed to remove media sink from presentation clock (%s)"), *WmfMedia::ResultToString(Result));
+			UE_LOG(LogWmfMedia, Error, TEXT("Failed to remove media sink from presentation clock: %s"), *WmfMedia::ResultToString(Result));
 			return Result;
 		}
 	}
@@ -272,7 +295,7 @@ STDMETHODIMP FWmfMediaSink::SetPresentationClock(__RPC__in_opt IMFPresentationCl
 
 		if (FAILED(Result))
 		{
-			UE_LOG(LogWmfMedia, Error, TEXT("Failed to add media sink to presentation clock (%s)"), *WmfMedia::ResultToString(Result));
+			UE_LOG(LogWmfMedia, Error, TEXT("Failed to add media sink to presentation clock: %s"), *WmfMedia::ResultToString(Result));
 			return Result;
 		}
 	}
@@ -299,6 +322,78 @@ STDMETHODIMP FWmfMediaSink::Shutdown()
 }
 
 
+/* IMFMediaSinkPreroll interface
+ *****************************************************************************/
+
+STDMETHODIMP FWmfMediaSink::NotifyPreroll(MFTIME hnsUpcomingStartTime)
+{
+	FScopeLock Lock(&CriticalSection);
+
+	if (!StreamSink.IsValid())
+	{
+		return MF_E_SHUTDOWN;
+	}
+
+	return StreamSink->Preroll();
+}
+
+
+/* IMFRateSupport interface
+ *****************************************************************************/
+
+STDMETHODIMP FWmfMediaSink::GetFastestRate(MFRATE_DIRECTION eDirection, BOOL fThin, _Out_ float* pflRate)
+{
+	if (pflRate == NULL)
+	{
+		return E_POINTER;
+	}
+
+	FScopeLock Lock(&CriticalSection);
+
+	if (!StreamSink.IsValid())
+	{
+		return MF_E_SHUTDOWN;
+	}
+
+	*pflRate = (eDirection == MFRATE_FORWARD) ? FLT_MAX : FLT_MIN;
+
+	return S_OK;
+}
+
+
+STDMETHODIMP FWmfMediaSink::GetSlowestRate(MFRATE_DIRECTION eDirection, BOOL fThin, _Out_ float* pflRate)
+{
+	if (pflRate == NULL)
+	{
+		return E_POINTER;
+	}
+
+	FScopeLock Lock(&CriticalSection);
+
+	if (!StreamSink.IsValid())
+	{
+		return MF_E_SHUTDOWN;
+	}
+
+	*pflRate = 0.0;
+
+	return S_OK;
+}
+
+
+STDMETHODIMP FWmfMediaSink::IsRateSupported(BOOL fThin, float flRate, __RPC__inout_opt float* pflNearestSupportedRate)
+{
+	FScopeLock Lock(&CriticalSection);
+
+	if (!StreamSink.IsValid())
+	{
+		return MF_E_SHUTDOWN;
+	}
+
+	return S_OK;
+}
+
+
 /* IUnknown interface
  *****************************************************************************/
 
@@ -318,7 +413,10 @@ STDMETHODIMP FWmfMediaSink::QueryInterface(REFIID RefID, void** Object)
 	static const QITAB QITab[] =
 	{
 		QITABENT(FWmfMediaSink, IMFClockStateSink),
+		QITABENT(FWmfMediaSink, IMFGetService),
 		QITABENT(FWmfMediaSink, IMFMediaSink),
+		QITABENT(FWmfMediaSink, IMFMediaSinkPreroll),
+		QITABENT(FWmfMediaSink, IMFRateSupport),
 		{ 0 }
 	};
 
diff --git a/Engine/Plugins/Media/WmfMedia/Source/WmfMedia/Private/Wmf/WmfMediaSink.h b/Engine/Plugins/Media/WmfMedia/Source/WmfMedia/Private/Wmf/WmfMediaSink.h
--- a/Engine/Plugins/Media/WmfMedia/Source/WmfMedia/Private/Wmf/WmfMediaSink.h
+++ b/Engine/Plugins/Media/WmfMedia/Source/WmfMedia/Private/Wmf/WmfMediaSink.h
@@ -17,7 +17,10 @@ class FWmfMediaStreamSink;
  */
 class FWmfMediaSink
 	: public IMFClockStateSink
+	, public IMFGetService
 	, public IMFMediaSink
+	, public IMFMediaSinkPreroll
+	, public IMFRateSupport
 {
 public:
 
@@ -44,6 +47,12 @@ public:
 	STDMETHODIMP OnClockStart(MFTIME hnsSystemTime, LONGLONG llClockStartOffset);
 	STDMETHODIMP OnClockStop(MFTIME hnsSystemTime);
 
+public:
+
+	//~ IMFGetService interface
+
+	STDMETHODIMP GetService(__RPC__in REFGUID guidService, __RPC__in REFIID riid, __RPC__deref_out_opt LPVOID* ppvObject);
+
 public:
 
 	// ~IMFMediaSink interface
@@ -58,6 +67,20 @@ public:
 	STDMETHODIMP SetPresentationClock(__RPC__in_opt IMFPresentationClock* pPresentationClock);
 	STDMETHODIMP Shutdown();
 
+public:
+
+	//~ IMFMediaSinkPreroll interface
+
+	STDMETHODIMP NotifyPreroll(MFTIME hnsUpcomingStartTime);
+
+public:
+
+	//~ IMFRateSupport interface
+
+	STDMETHODIMP GetFastestRate(MFRATE_DIRECTION eDirection, BOOL fThin, _Out_ float* pflRate);
+	STDMETHODIMP GetSlowestRate(MFRATE_DIRECTION eDirection, BOOL fThin, _Out_ float* pflRate);
+	STDMETHODIMP IsRateSupported(BOOL fThin, float flRate, __RPC__inout_opt float* pflNearestSupportedRate);
+
 public:
 
 	//~ IUnknown interface
