Copying //UE4/Dev-Sequencer to //UE4/Main (Source: //UE4/Dev-Sequencer @ 3631098)

diff --git a/Engine/Plugins/Media/WmfMedia/Source/WmfMedia/Private/Wmf/WmfMediaStreamSink.cpp b/Engine/Plugins/Media/WmfMedia/Source/WmfMedia/Private/Wmf/WmfMediaStreamSink.cpp
--- a/Engine/Plugins/Media/WmfMedia/Source/WmfMedia/Private/Wmf/WmfMediaStreamSink.cpp
+++ b/Engine/Plugins/Media/WmfMedia/Source/WmfMedia/Private/Wmf/WmfMediaStreamSink.cpp
@@ -11,25 +11,95 @@
 #include "WmfMediaUtils.h"
 
 
+/* FWmfMediaStreamSink static functions
+ *****************************************************************************/
+
+bool FWmfMediaStreamSink::Create(const GUID& MajorType, TComPtr<FWmfMediaStreamSink>& OutSink)
+{
+	TComPtr<FWmfMediaStreamSink> StreamSink = new FWmfMediaStreamSink(MajorType, 1);
+	TComPtr<FWmfMediaSink> MediaSink = new FWmfMediaSink();
+
+	if (!MediaSink->Initialize(*StreamSink))
+	{
+		return false;
+	}
+
+	OutSink = StreamSink;
+
+	return true;
+}
+
+
 /* FWmfMediaStreamSink structors
  *****************************************************************************/
 
 FWmfMediaStreamSink::FWmfMediaStreamSink(const GUID& InMajorType, DWORD InStreamId)
-	: RefCount(0)
+	: Prerolling(false)
+	, RefCount(0)
 	, StreamId(InStreamId)
 	, StreamType(InMajorType)
-{ }
+{
+	UE_LOG(LogWmfMedia, Verbose, TEXT("StreamSink %p: Created with stream type %s"), this, *WmfMedia::MajorTypeToString(StreamType));
+}
 
 
 FWmfMediaStreamSink::~FWmfMediaStreamSink()
 {
 	check(RefCount == 0);
+
+	UE_LOG(LogWmfMedia, Verbose, TEXT("StreamSink %p: Destroyed"), this);
 }
 
 
 /* FWmfMediaStreamSink interface
  *****************************************************************************/
 
+bool FWmfMediaStreamSink::GetNextSample(const TRange<FTimespan>& SampleRange, FWmfMediaStreamSinkSample& OutSample)
+{
+	if (SampleRange.IsEmpty())
+	{
+		return false; // nothing to play
+	}
+
+	FScopeLock Lock(&CriticalSection);
+
+	FQueuedSample QueuedSample;
+
+	while (SampleQueue.Peek(QueuedSample))
+	{
+		if (QueuedSample.MediaType.IsValid())
+		{
+			check(QueuedSample.Sample.IsValid());
+
+			if (!SampleRange.Contains(FTimespan(QueuedSample.Time)))
+			{
+				return false; // no new sample needed
+			}
+
+			check(SampleQueue.Pop());
+
+			OutSample.MediaType = QueuedSample.MediaType;
+			OutSample.Sample = QueuedSample.Sample;
+
+			return true;
+		}
+
+		SampleQueue.Pop();
+
+		// process pending marker
+		QueueEvent(MEStreamSinkMarker, GUID_NULL, S_OK, QueuedSample.MarkerContext);
+		PropVariantClear(QueuedSample.MarkerContext);
+		delete QueuedSample.MarkerContext;
+
+		UE_LOG(LogWmfMedia, Verbose, TEXT("StreamSink %p: Processed marker (%s)"), this, *WmfMedia::MarkerTypeToString(QueuedSample.MarkerType));
+
+		continue;
+	}
+
+	return false;
+}
+
+
 bool FWmfMediaStreamSink::Initialize(FWmfMediaSink& InOwner)
 {
 	FScopeLock Lock(&CriticalSection);
@@ -38,7 +108,7 @@ bool FWmfMediaStreamSink::Initialize(FWmfMediaSink& InOwner)
 
 	if (FAILED(Result))
 	{
-		UE_LOG(LogWmfMedia, VeryVerbose, TEXT("Failed to create event queue for stream sink (%s)"), *WmfMedia::ResultToString(Result));
+		UE_LOG(LogWmfMedia, Verbose, TEXT("StreamSink %p: Failed to create event queue for stream sink: %s"), this, *WmfMedia::ResultToString(Result));
 		return false;
 	}
 
@@ -48,17 +118,34 @@ bool FWmfMediaStreamSink::Initialize(FWmfMediaSink& InOwner)
 }
 
 
-FOnWmfMediaStreamSinkBuffer& FWmfMediaStreamSink::OnBuffer()
+HRESULT FWmfMediaStreamSink::Pause()
+{
+	FScopeLock Lock(&CriticalSection);
+
+	return QueueEvent(MEStreamSinkPaused, GUID_NULL, S_OK, NULL);
+}
+
+
+HRESULT FWmfMediaStreamSink::Preroll()
 {
 	FScopeLock Lock(&CriticalSection);
-	return BufferDelegate;
+
+	if (!EventQueue.IsValid())
+	{
+		return MF_E_SHUTDOWN;
+	}
+
+	Prerolling = true;
+
+	return QueueEvent(MEStreamSinkRequestSample, GUID_NULL, S_OK, NULL);
 }
 
 
-FOnWmfMediaStreamSinkSample& FWmfMediaStreamSink::OnSample()
+HRESULT FWmfMediaStreamSink::Restart()
 {
 	FScopeLock Lock(&CriticalSection);
-	return SampleDelegate;
+
+	return QueueEvent(MEStreamSinkStarted, GUID_NULL, S_OK, NULL);
 }
 
 
@@ -74,6 +161,40 @@ void FWmfMediaStreamSink::Shutdown()
 }
 
 
+HRESULT FWmfMediaStreamSink::Start()
+{
+	FScopeLock Lock(&CriticalSection);
+
+	HRESULT Result = QueueEvent(MEStreamSinkStarted, GUID_NULL, S_OK, NULL);
+
+	if (FAILED(Result))
+	{
+		return Result;
+	}
+
+	return QueueEvent(MEStreamSinkRequestSample, GUID_NULL, S_OK, NULL);
+}
+
+
+HRESULT FWmfMediaStreamSink::Stop()
+{
+	Flush();
+
+	FScopeLock Lock(&CriticalSection);
+
+	return QueueEvent(MEStreamSinkStopped, GUID_NULL, S_OK, NULL);
+}
+
+
+/* IMFGetService interface
+ *****************************************************************************/
+
+STDMETHODIMP FWmfMediaStreamSink::GetService(__RPC__in REFGUID guidService, __RPC__in REFIID riid, __RPC__deref_out_opt LPVOID* ppvObject)
+{
+	return Owner->GetService(guidService, riid, ppvObject);
+}
+
+
 /* IMFMediaEventGenerator interface
  *****************************************************************************/
 
@@ -272,6 +393,8 @@ STDMETHODIMP FWmfMediaStreamSink::IsMediaTypeSupported(IMFMediaType* pMediaType,
 		return E_POINTER;
 	}
 
+	UE_LOG(LogWmfMedia, VeryVerbose, TEXT("StreamSink %p: Checking if media type is supported:\n%s"), *WmfMedia::DumpAttributes(*pMediaType));
+
 	FScopeLock Lock(&CriticalSection);
 
 	if (!EventQueue.IsValid())
@@ -292,6 +415,7 @@ STDMETHODIMP FWmfMediaStreamSink::IsMediaTypeSupported(IMFMediaType* pMediaType,
 
 	if (MajorType != StreamType)
 	{
+		UE_LOG(LogWmfMedia, VeryVerbose, TEXT("StreamSink %p: Media type doesn't match stream type %s"), this, *WmfMedia::MajorTypeToString(StreamType));
 		return MF_E_INVALIDMEDIATYPE;
 	}
 
@@ -310,10 +434,13 @@ STDMETHODIMP FWmfMediaStreamSink::IsMediaTypeSupported(IMFMediaType* pMediaType,
 
 		if (SUCCEEDED(Result) && ((OutFlags & CompareFlags) == CompareFlags))
 		{
+			UE_LOG(LogWmfMedia, VeryVerbose, TEXT("StreamSink %p: Media type is supported"), this, *WmfMedia::MajorTypeToString(MajorType));
 			return S_OK;
 		}
 	}
 
+	UE_LOG(LogWmfMedia, VeryVerbose, TEXT("StreamSink %p: Media type is not supported"), this);
+
 	return MF_E_INVALIDMEDIATYPE;
 }
 
@@ -325,6 +452,8 @@ STDMETHODIMP FWmfMediaStreamSink::SetCurrentMediaType(IMFMediaType* pMediaType)
 		return E_POINTER;
 	}
 
+	UE_LOG(LogWmfMedia, VeryVerbose, TEXT("StreamSink %p: Setting current media type:\n%s"), *WmfMedia::DumpAttributes(*pMediaType));
+
 	FScopeLock Lock(&CriticalSection);
 
 	if (!EventQueue.IsValid())
@@ -336,9 +465,12 @@ STDMETHODIMP FWmfMediaStreamSink::SetCurrentMediaType(IMFMediaType* pMediaType)
 
 	if (FAILED(Result))
 	{
+		UE_LOG(LogWmfMedia, VeryVerbose, TEXT("StreamSink %p: Tried to set unsupported media type"), this);
 		return Result;
 	}
 
+	UE_LOG(LogWmfMedia, VeryVerbose, TEXT("StreamSink %p: Current media type set"), this);
+
 	CurrentMediaType = pMediaType;
 
 	return S_OK;
@@ -357,6 +489,22 @@ STDMETHODIMP FWmfMediaStreamSink::Flush()
 		return MF_E_SHUTDOWN;
 	}
 
+	UE_LOG(LogWmfMedia, Verbose, TEXT("StreamSink %p: Flushing samples & markers"), this);
+
+	FQueuedSample QueuedSample;
+
+	while (SampleQueue.Dequeue(QueuedSample))
+	{
+		if (QueuedSample.MediaType.IsValid())
+		{
+			continue;
+		}
+
+		// notify WMF that flushed markers haven't been processed
+		QueueEvent(MEStreamSinkMarker, GUID_NULL, E_ABORT, QueuedSample.MarkerContext);
+		PropVariantClear(QueuedSample.MarkerContext);
+		delete QueuedSample.MarkerContext;
+	}
 
 	return S_OK;
 }
@@ -430,6 +578,25 @@ STDMETHODIMP FWmfMediaStreamSink::PlaceMarker(MFSTREAMSINK_MARKER_TYPE eMarkerTy
 		return MF_E_SHUTDOWN;
 	}
 
+	UE_LOG(LogWmfMedia, Verbose, TEXT("StreamSink %p: Placing marker (%s)"), this, *WmfMedia::MarkerTypeToString(eMarkerType));
+
+	PROPVARIANT* MarkerContext = new PROPVARIANT;
+
+	if (pvarContextValue != NULL)
+	{
+		HRESULT Result = ::PropVariantCopy(MarkerContext, pvarContextValue);
+
+		if (FAILED(Result))
+		{
+			UE_LOG(LogWmfMedia, Verbose, TEXT("StreamSink %p: Failed to copy marker context: %s"), this, *WmfMedia::ResultToString(Result));
+			delete MarkerContext;
+
+			return Result;
+		}
+	}
+
+	SampleQueue.Enqueue({ eMarkerType, MarkerContext, NULL, NULL, 0 });
+
 	return S_OK;
 }
 
@@ -448,91 +615,37 @@ STDMETHODIMP FWmfMediaStreamSink::ProcessSample(__RPC__in_opt IMFSample* pSample
 		return MF_E_SHUTDOWN;
 	}
 
-	DWORD NumBuffers;
-	{
-		const HRESULT Result = pSample->GetBufferCount(&NumBuffers);
-
-		if (FAILED(Result))
-		{
-			UE_LOG(LogWmfMedia, VeryVerbose, TEXT("Failed to get buffer count from media sample (%s)"), *WmfMedia::ResultToString(Result));
-			return Result;
-		}
-	}
-
-	FTimespan Duration;
+	if (!CurrentMediaType.IsValid())
 	{
-		LONGLONG SampleDuration = 0;
-		const HRESULT Result = pSample->GetSampleDuration(&SampleDuration);
-
-		if (FAILED(Result))
-		{
-			UE_LOG(LogWmfMedia, VeryVerbose, TEXT("Failed to get sample duration from media sample (%s)"), *WmfMedia::ResultToString(Result));
-			return Result;
-		}
-
-		Duration = FTimespan(SampleDuration);
+		UE_LOG(LogWmfMedia, VeryVerbose, TEXT("StreamSink %p: Stream received a sample while not having a valid media type set"), this);
+		return MF_E_INVALIDMEDIATYPE;
 	}
 
-	DWORD Flags;
+	// get sample time
+	LONGLONG Time = 0;
 	{
-		const HRESULT Result = pSample->GetSampleFlags(&Flags);
+		const HRESULT Result = pSample->GetSampleTime(&Time);
 
 		if (FAILED(Result))
 		{
-			UE_LOG(LogWmfMedia, VeryVerbose, TEXT("Failed to get sample flags from media sample (%s)"), *WmfMedia::ResultToString(Result));
+			UE_LOG(LogWmfMedia, VeryVerbose, TEXT("Failed to get time from sink sample: %s"), *WmfMedia::ResultToString(Result));
 			return Result;
 		}
 	}
 
-	FTimespan Time;
-	{
-		LONGLONG SampleTime = 0;
-		const HRESULT Result = pSample->GetSampleTime(&SampleTime);
-
-		if (FAILED(Result))
-		{
-			UE_LOG(LogWmfMedia, VeryVerbose, TEXT("Failed to get sample time from media sample (%s)"), *WmfMedia::ResultToString(Result));
-			return Result;
-		}
+	SampleQueue.Enqueue({ MFSTREAMSINK_MARKER_DEFAULT, NULL, CurrentMediaType, pSample, Time });
 
-		Time = FTimespan(SampleTime);
-	}
-
-	if (BufferDelegate.IsBound())
+	// finish pre-rolling
+	if (Prerolling)
 	{
-		TComPtr<IMFMediaBuffer> Buffer = NULL;
-		{
-			const HRESULT Result = pSample->ConvertToContiguousBuffer(&Buffer);
-
-			if (FAILED(Result))
-			{
-				UE_LOG(LogWmfMedia, VeryVerbose, TEXT("Failed to get contiguous buffer from media sample (%s)"), *WmfMedia::ResultToString(Result));
-				return Result;
-			}
-		}
-
-		BYTE* ByteBuffer = NULL;
-		DWORD Size = 0;
-		{
-			const HRESULT Result = Buffer->Lock(&ByteBuffer, NULL, &Size);
+		UE_LOG(LogWmfMedia, VeryVerbose, TEXT("StreamSink %p: Preroll complete"), this);
+		Prerolling = false;
 
-			if (FAILED(Result))
-			{
-				UE_LOG(LogWmfMedia, VeryVerbose, TEXT("Failed to lock contiguous buffer from media sample (%s)"), *WmfMedia::ResultToString(Result));
-				return Result;
-			}
-		}
-
-		BufferDelegate.Execute(ByteBuffer, Size, Time, Duration, Flags);
-		Buffer->Unlock();
+		return QueueEvent(MEStreamSinkPrerolled, GUID_NULL, S_OK, NULL);
 	}
 
-	if (SampleDelegate.IsBound())
-	{
-		SampleDelegate.Execute(*pSample, NumBuffers, Time, Duration, Flags);
-	}
-
-	return S_OK;
+	// request another sample
+	return QueueEvent(MEStreamSinkRequestSample, GUID_NULL, S_OK, NULL);
 }
 
 
@@ -554,6 +667,8 @@ STDMETHODIMP FWmfMediaStreamSink::QueryInterface(REFIID RefID, void** Object)
 {
 	static const QITAB QITab[] =
 	{
+		QITABENT(FWmfMediaStreamSink, IMFGetService),
+		QITABENT(FWmfMediaStreamSink, IMFMediaTypeHandler),
 		QITABENT(FWmfMediaStreamSink, IMFStreamSink),
 		{ 0 }
 	};
diff --git a/Engine/Plugins/Media/WmfMedia/Source/WmfMedia/Private/Wmf/WmfMediaStreamSink.h b/Engine/Plugins/Media/WmfMedia/Source/WmfMedia/Private/Wmf/WmfMediaStreamSink.h
--- a/Engine/Plugins/Media/WmfMedia/Source/WmfMedia/Private/Wmf/WmfMediaStreamSink.h
+++ b/Engine/Plugins/Media/WmfMedia/Source/WmfMedia/Private/Wmf/WmfMediaStreamSink.h
@@ -7,6 +7,7 @@
 #if WMFMEDIA_SUPPORTED_PLATFORM
 
 #include "CoreTypes.h"
+#include "Containers/Queue.h"
 #include "Delegates/Delegate.h"
 #include "HAL/CriticalSection.h"
 #include "Misc/Timespan.h"
@@ -14,39 +15,38 @@
 class FWmfMediaSink;
 
 
-/** Declares a delegate that gets fired when the sink received a new sample. */
-DECLARE_DELEGATE_FiveParams(FOnWmfMediaStreamSinkSample,
-	IMFSample& /*Sample*/,
-	DWORD /*NumBuffers*/,
-	FTimespan /*Time*/,
-	FTimespan /*Duration*/,
-	DWORD /*Flags*/
-);
-
-
-/** Declares a delegate that gets fired when the sink received a new sample. */
-DECLARE_DELEGATE_FiveParams(FOnWmfMediaStreamSinkBuffer,
-	const BYTE* /*Buffer*/,
-	DWORD /*Size*/,
-	FTimespan /*Time*/,
-	FTimespan /*Duration*/,
-	DWORD /*Flags*/
-);
-
-
-namespace WmfMediaStreamSink
+/**
+ * Structure for media samples queued in a stream sink.
+ */
+struct FWmfMediaStreamSinkSample
 {
-	const DWORD FixedStreamId = 1;
-}
+	/** The sample's media type. */
+	TComPtr<IMFMediaType> MediaType;
+
+	/** The media sample. */
+	TComPtr<IMFSample> Sample;
+};
 
 
 /**
  * Implements a stream sink object for the WMF pipeline.
  */
 class FWmfMediaStreamSink
-	: public IMFMediaTypeHandler
+	: public IMFGetService
+	, public IMFMediaTypeHandler
 	, public IMFStreamSink
 {
+public:
+
+	/**
+	 * Creates a stream sink for the specified major type.
+	 *
+	 * @param MajorType The sink's major type.
+	 * @param OutSink Will contain the created sink.
+	 * @return true on success, false otherwise.
+	 */
+	static bool Create(const GUID& MajorType, TComPtr<FWmfMediaStreamSink>& OutSink);
+
 public:
 
 	/**
@@ -59,6 +59,15 @@ public:
 
 public:
 
+	/**
+	 * Get the next sample in the queue.
+	 *
+	 * @param SampleRange Time range of samples that should be returned.
+	 * @param OutSample Will contain the sample.
+	 * @return true if a sample was returned, false if the queue is empty.
+	 */
+	bool GetNextSample(const TRange<FTimespan>& SampleRange, FWmfMediaStreamSinkSample& OutSample);
+
 	/**
 	 * Initialize this sink.
 	 *
@@ -68,11 +77,35 @@ public:
 	 */
 	bool Initialize(FWmfMediaSink& InOwner);
 
-	/** Get a delegate that is executed when the sink received a new sample. */
-	FOnWmfMediaStreamSinkBuffer& OnBuffer();
+	/**
+	 * Pause the stream.
+	 *
+	 * This method is called by the owner media sink.
+	 *
+	 * @return Result code.
+	 * @see Preroll, Restart, Start, Stop
+	 */
+	HRESULT Pause();
 
-	/** Get a delegate that is executed when the sink received a new sample. */
-	FOnWmfMediaStreamSinkSample& OnSample();
+	/**
+	 * Preroll the sink.
+	 *
+	 * This method is called by the owner media sink.
+	 *
+	 * @return Result code.
+	 * @see Pause, Restart, Start, Stop
+	 */
+	HRESULT Preroll();
+
+	/**
+	 * Restart the stream.
+	 *
+	 * This method is called by the owner media sink.
+	 *
+	 * @return Result code.
+	 * @see Pause, Preroll, Start, Stop
+	 */
+	HRESULT Restart();
 
 	/**
 	 * Shut down this sink.
@@ -81,6 +114,32 @@ public:
 	 */
 	void Shutdown();
 
+	/**
+	 * Start the sink.
+	 *
+	 * This method is called by the owner media sink.
+	 *
+	 * @return Result code.
+	 * @see Pause, Preroll, Restart, Stop
+	 */
+	HRESULT Start();
+
+	/**
+	 * Stop the stream.
+	 *
+	 * This method is called by the owner media sink.
+	 *
+	 * @return Result code.
+	 * @see Pause, Preroll, Restart, Start
+	 */
+	HRESULT Stop();
+
+public:
+
+	//~ IMFGetService interface
+
+	STDMETHODIMP GetService(__RPC__in REFGUID guidService, __RPC__in REFIID riid, __RPC__deref_out_opt LPVOID* ppvObject);
+
 public:
 
 	//~ IMFMediaEventGenerator interface
@@ -127,9 +186,6 @@ private:
 
 private:
 
-	/** The delegate that gets fired when the sink received a new sample. */
-	FOnWmfMediaStreamSinkBuffer BufferDelegate;
-
 	/** Critical section for synchronizing access to this sink. */
 	FCriticalSection CriticalSection;
 
@@ -142,17 +198,41 @@ private:
 	/** The media sink that owns this stream sink. */
 	TComPtr<FWmfMediaSink> Owner;
 
+	/** Whether the sink is currently prerolling samples. */
+	bool Prerolling;
+
 	/** Holds a reference counter for this instance. */
 	int32 RefCount;
 
-	/** The delegate that gets fired when the sink received a new sample. */
-	FOnWmfMediaStreamSinkSample SampleDelegate;
-
 	/** The stream identifier (currently fixed). */
 	DWORD StreamId;
 
 	/** The sink's major media type. */
 	const GUID StreamType;
+
+private:
+
+	/** Structure for queued media samples & markers. */
+	struct FQueuedSample
+	{
+		/** Stream marker type. */
+		MFSTREAMSINK_MARKER_TYPE MarkerType;
+
+		/** Stream marker context. */
+		PROPVARIANT* MarkerContext;
+
+		/** The sample's media type. */
+		TComPtr<IMFMediaType> MediaType;
+
+		/** The media sample. */
+		TComPtr<IMFSample> Sample;
+
+		/** Sample time. */
+		LONGLONG Time;
+	};
+
+	/** Media sample queue. */
+	TQueue<FQueuedSample> SampleQueue;
 };
 
 #endif
diff --git a/Engine/Plugins/Media/WmfMedia/Source/WmfMedia/Private/Wmf/WmfMediaUtils.cpp b/Engine/Plugins/Media/WmfMedia/Source/WmfMedia/Private/Wmf/WmfMediaUtils.cpp
--- a/Engine/Plugins/Media/WmfMedia/Source/WmfMedia/Private/Wmf/WmfMediaUtils.cpp
+++ b/Engine/Plugins/Media/WmfMedia/Source/WmfMedia/Private/Wmf/WmfMediaUtils.cpp
@@ -66,7 +66,6 @@ namespace WmfMedia
 	GUID const* const SupportedMajorTypes[] = {
 		&MFMediaType_Audio,
 		&MFMediaType_Binary,
-		&MFMediaType_Image,
 		&MFMediaType_SAMI,
 		&MFMediaType_Video
 	};
@@ -95,8 +94,8 @@ namespace WmfMedia
 //		&MFVideoFormat_A2R10G10B10,
 //		&MFVideoFormat_A16B16G16R16F,
 //		&MFVideoFormat_ARGB32
-		&MFVideoFormat_RGB24,
 		&MFVideoFormat_RGB32,
+		&MFVideoFormat_RGB24,
 //		&MFVideoFormat_RGB555,
 //		&MFVideoFormat_RGB565,
 //		&MFVideoFormat_RGB8,
@@ -157,6 +156,8 @@ namespace WmfMedia
 			return;
 		}
 
+		Initialized = true;
+
 		// initialize audio types
 		for (const FAudioFormat& Format : SupportedAudioFormats)
 		{
@@ -171,8 +172,7 @@ namespace WmfMedia
 					FAILED(MediaType->SetUINT32(MF_MT_AUDIO_BITS_PER_SAMPLE, Format.BitsPerSample)) ||
 					FAILED(MediaType->SetUINT32(MF_MT_AUDIO_NUM_CHANNELS, NumChannels)))
 				{
-					UE_LOG(LogWmfMedia, Error, TEXT("Failed to initialize supported audio types"));
-					return;
+					UE_LOG(LogWmfMedia, Error, TEXT("Failed to initialize supported audio type: %s, %i bits/sample, %i channels"), *SubTypeToString(*Format.SubType), Format.BitsPerSample, NumChannels);
 				}
 
 				SupportedAudioTypes.Add(MediaType);
@@ -188,7 +188,6 @@ namespace WmfMedia
 				FAILED(MediaType->SetUINT32(MF_MT_ALL_SAMPLES_INDEPENDENT, TRUE)))
 			{
 				UE_LOG(LogWmfMedia, Error, TEXT("Failed to initialize supported binary types"));
-				return;
 			}
 
 			SupportedBinaryTypes.Add(MediaType);
@@ -203,7 +202,6 @@ namespace WmfMedia
 				FAILED(MediaType->SetUINT32(MF_MT_ALL_SAMPLES_INDEPENDENT, TRUE)))
 			{
 				UE_LOG(LogWmfMedia, Error, TEXT("Failed to initialize supported SAMI types"));
-				return;
 			}
 
 			SupportedSamiTypes.Add(MediaType);
@@ -217,10 +215,10 @@ namespace WmfMedia
 			if (FAILED(::MFCreateMediaType(&MediaType)) ||
 				FAILED(MediaType->SetGUID(MF_MT_MAJOR_TYPE, MFMediaType_Video)) ||
 				FAILED(MediaType->SetGUID(MF_MT_SUBTYPE, *Format)) ||
-				FAILED(MediaType->SetUINT32(MF_MT_ALL_SAMPLES_INDEPENDENT, TRUE)))
+				FAILED(MediaType->SetUINT32(MF_MT_ALL_SAMPLES_INDEPENDENT, TRUE)) ||
+				FAILED(MediaType->SetUINT32(MF_MT_INTERLACE_MODE, (UINT32)MFVideoInterlace_Progressive)))
 			{
-				UE_LOG(LogWmfMedia, Error, TEXT("Failed to initialize supported video types"));
-				return;
+				UE_LOG(LogWmfMedia, Error, TEXT("Failed to initialize supported video type: %s"), *SubTypeToString(*Format));
 			}
 
 			SupportedVideoTypes.Add(MediaType);
@@ -234,18 +232,91 @@ namespace WmfMedia
 
 namespace WmfMedia
 {
+	FString AttributeToString(const GUID& Guid)
+	{
+		if (Guid == MF_MT_MAJOR_TYPE) return TEXT("MF_MT_MAJOR_TYPE");
+		if (Guid == MF_MT_MAJOR_TYPE) return TEXT("MF_MT_MAJOR_TYPE");
+		if (Guid == MF_MT_SUBTYPE) return TEXT("MF_MT_SUBTYPE");
+		if (Guid == MF_MT_ALL_SAMPLES_INDEPENDENT) return TEXT("MF_MT_ALL_SAMPLES_INDEPENDENT");
+		if (Guid == MF_MT_FIXED_SIZE_SAMPLES) return TEXT("MF_MT_FIXED_SIZE_SAMPLES");
+		if (Guid == MF_MT_COMPRESSED) return TEXT("MF_MT_COMPRESSED");
+		if (Guid == MF_MT_SAMPLE_SIZE) return TEXT("MF_MT_SAMPLE_SIZE");
+		if (Guid == MF_MT_WRAPPED_TYPE) return TEXT("MF_MT_WRAPPED_TYPE");
+		if (Guid == MF_MT_AUDIO_NUM_CHANNELS) return TEXT("MF_MT_AUDIO_NUM_CHANNELS");
+		if (Guid == MF_MT_AUDIO_SAMPLES_PER_SECOND) return TEXT("MF_MT_AUDIO_SAMPLES_PER_SECOND");
+		if (Guid == MF_MT_AUDIO_FLOAT_SAMPLES_PER_SECOND) return TEXT("MF_MT_AUDIO_FLOAT_SAMPLES_PER_SECOND");
+		if (Guid == MF_MT_AUDIO_AVG_BYTES_PER_SECOND) return TEXT("MF_MT_AUDIO_AVG_BYTES_PER_SECOND");
+		if (Guid == MF_MT_AUDIO_BLOCK_ALIGNMENT) return TEXT("MF_MT_AUDIO_BLOCK_ALIGNMENT");
+		if (Guid == MF_MT_AUDIO_BITS_PER_SAMPLE) return TEXT("MF_MT_AUDIO_BITS_PER_SAMPLE");
+		if (Guid == MF_MT_AUDIO_VALID_BITS_PER_SAMPLE) return TEXT("MF_MT_AUDIO_VALID_BITS_PER_SAMPLE");
+		if (Guid == MF_MT_AUDIO_SAMPLES_PER_BLOCK) return TEXT("MF_MT_AUDIO_SAMPLES_PER_BLOCK");
+		if (Guid == MF_MT_AUDIO_CHANNEL_MASK) return TEXT("MF_MT_AUDIO_CHANNEL_MASK");
+		if (Guid == MF_MT_AUDIO_FOLDDOWN_MATRIX) return TEXT("MF_MT_AUDIO_FOLDDOWN_MATRIX");
+		if (Guid == MF_MT_AUDIO_WMADRC_PEAKREF) return TEXT("MF_MT_AUDIO_WMADRC_PEAKREF");
+		if (Guid == MF_MT_AUDIO_WMADRC_PEAKTARGET) return TEXT("MF_MT_AUDIO_WMADRC_PEAKTARGET");
+		if (Guid == MF_MT_AUDIO_WMADRC_AVGREF) return TEXT("MF_MT_AUDIO_WMADRC_AVGREF");
+		if (Guid == MF_MT_AUDIO_WMADRC_AVGTARGET) return TEXT("MF_MT_AUDIO_WMADRC_AVGTARGET");
+		if (Guid == MF_MT_AUDIO_PREFER_WAVEFORMATEX) return TEXT("MF_MT_AUDIO_PREFER_WAVEFORMATEX");
+		if (Guid == MF_MT_AAC_PAYLOAD_TYPE) return TEXT("MF_MT_AAC_PAYLOAD_TYPE");
+		if (Guid == MF_MT_AAC_AUDIO_PROFILE_LEVEL_INDICATION) return TEXT("MF_MT_AAC_AUDIO_PROFILE_LEVEL_INDICATION");
+		if (Guid == MF_MT_FRAME_SIZE) return TEXT("MF_MT_FRAME_SIZE");
+		if (Guid == MF_MT_FRAME_RATE) return TEXT("MF_MT_FRAME_RATE");
+		if (Guid == MF_MT_FRAME_RATE_RANGE_MAX) return TEXT("MF_MT_FRAME_RATE_RANGE_MAX");
+		if (Guid == MF_MT_FRAME_RATE_RANGE_MIN) return TEXT("MF_MT_FRAME_RATE_RANGE_MIN");
+		if (Guid == MF_MT_PIXEL_ASPECT_RATIO) return TEXT("MF_MT_PIXEL_ASPECT_RATIO");
+		if (Guid == MF_MT_DRM_FLAGS) return TEXT("MF_MT_DRM_FLAGS");
+		if (Guid == MF_MT_PAD_CONTROL_FLAGS) return TEXT("MF_MT_PAD_CONTROL_FLAGS");
+		if (Guid == MF_MT_SOURCE_CONTENT_HINT) return TEXT("MF_MT_SOURCE_CONTENT_HINT");
+		if (Guid == MF_MT_VIDEO_CHROMA_SITING) return TEXT("MF_MT_VIDEO_CHROMA_SITING");
+		if (Guid == MF_MT_INTERLACE_MODE) return TEXT("MF_MT_INTERLACE_MODE");
+		if (Guid == MF_MT_TRANSFER_FUNCTION) return TEXT("MF_MT_TRANSFER_FUNCTION");
+		if (Guid == MF_MT_VIDEO_PRIMARIES) return TEXT("MF_MT_VIDEO_PRIMARIES");
+		if (Guid == MF_MT_CUSTOM_VIDEO_PRIMARIES) return TEXT("MF_MT_CUSTOM_VIDEO_PRIMARIES");
+		if (Guid == MF_MT_YUV_MATRIX) return TEXT("MF_MT_YUV_MATRIX");
+		if (Guid == MF_MT_VIDEO_LIGHTING) return TEXT("MF_MT_VIDEO_LIGHTING");
+		if (Guid == MF_MT_VIDEO_NOMINAL_RANGE) return TEXT("MF_MT_VIDEO_NOMINAL_RANGE");
+		if (Guid == MF_MT_GEOMETRIC_APERTURE) return TEXT("MF_MT_GEOMETRIC_APERTURE");
+		if (Guid == MF_MT_MINIMUM_DISPLAY_APERTURE) return TEXT("MF_MT_MINIMUM_DISPLAY_APERTURE");
+		if (Guid == MF_MT_PAN_SCAN_APERTURE) return TEXT("MF_MT_PAN_SCAN_APERTURE");
+		if (Guid == MF_MT_PAN_SCAN_ENABLED) return TEXT("MF_MT_PAN_SCAN_ENABLED");
+		if (Guid == MF_MT_AVG_BITRATE) return TEXT("MF_MT_AVG_BITRATE");
+		if (Guid == MF_MT_AVG_BIT_ERROR_RATE) return TEXT("MF_MT_AVG_BIT_ERROR_RATE");
+		if (Guid == MF_MT_MAX_KEYFRAME_SPACING) return TEXT("MF_MT_MAX_KEYFRAME_SPACING");
+		if (Guid == MF_MT_DEFAULT_STRIDE) return TEXT("MF_MT_DEFAULT_STRIDE");
+		if (Guid == MF_MT_PALETTE) return TEXT("MF_MT_PALETTE");
+		if (Guid == MF_MT_USER_DATA) return TEXT("MF_MT_USER_DATA");
+		if (Guid == MF_MT_AM_FORMAT_TYPE) return TEXT("MF_MT_AM_FORMAT_TYPE");
+		if (Guid == MF_MT_MPEG_START_TIME_CODE) return TEXT("MF_MT_MPEG_START_TIME_CODE");
+		if (Guid == MF_MT_MPEG2_PROFILE) return TEXT("MF_MT_MPEG2_PROFILE");
+		if (Guid == MF_MT_MPEG2_LEVEL) return TEXT("MF_MT_MPEG2_LEVEL");
+		if (Guid == MF_MT_MPEG2_FLAGS) return TEXT("MF_MT_MPEG2_FLAGS");
+		if (Guid == MF_MT_MPEG_SEQUENCE_HEADER) return TEXT("MF_MT_MPEG_SEQUENCE_HEADER");
+		if (Guid == MF_MT_DV_AAUX_SRC_PACK_0) return TEXT("MF_MT_DV_AAUX_SRC_PACK_0");
+		if (Guid == MF_MT_DV_AAUX_CTRL_PACK_0) return TEXT("MF_MT_DV_AAUX_CTRL_PACK_0");
+		if (Guid == MF_MT_DV_AAUX_SRC_PACK_1) return TEXT("MF_MT_DV_AAUX_SRC_PACK_1");
+		if (Guid == MF_MT_DV_AAUX_CTRL_PACK_1) return TEXT("MF_MT_DV_AAUX_CTRL_PACK_1");
+		if (Guid == MF_MT_DV_VAUX_SRC_PACK) return TEXT("MF_MT_DV_VAUX_SRC_PACK");
+		if (Guid == MF_MT_DV_VAUX_CTRL_PACK) return TEXT("MF_MT_DV_VAUX_CTRL_PACK");
+		if (Guid == MF_MT_ARBITRARY_HEADER) return TEXT("MF_MT_ARBITRARY_HEADER");
+		if (Guid == MF_MT_ARBITRARY_FORMAT) return TEXT("MF_MT_ARBITRARY_FORMAT");
+		if (Guid == MF_MT_IMAGE_LOSS_TOLERANT) return TEXT("MF_MT_IMAGE_LOSS_TOLERANT");
+		if (Guid == MF_MT_MPEG4_SAMPLE_DESCRIPTION) return TEXT("MF_MT_MPEG4_SAMPLE_DESCRIPTION");
+		if (Guid == MF_MT_MPEG4_CURRENT_SAMPLE_ENTRY) return TEXT("MF_MT_MPEG4_CURRENT_SAMPLE_ENTRY");
+		if (Guid == MF_MT_ORIGINAL_4CC) return TEXT("MF_MT_ORIGINAL_4CC");
+		if (Guid == MF_MT_ORIGINAL_WAVE_FORMAT_TAG) return TEXT("MF_MT_ORIGINAL_WAVE_FORMAT_TAG");
+
+		// unknown identifier
+		return GuidToString(Guid);
+	}
+
+
 	FString CaptureDeviceRoleToString(ERole Role)
 	{
 		switch (Role)
 		{
-		case eCommunications:
-			return TEXT("Communications");
-
-		case eConsole:
-			return TEXT("Console");
-
-		case eMultimedia:
-			return TEXT("Multimedia");
+		case eCommunications: return TEXT("Communications");
+		case eConsole: return TEXT("Console");
+		case eMultimedia: return TEXT("Multimedia");
 
 		default:
 			return TEXT("Unknown");
@@ -253,6 +324,23 @@ namespace WmfMedia
 	}
 
 
+	HRESULT CopyAttribute(IMFAttributes* Src, IMFAttributes* Dest, const GUID& Key)
+	{
+		PROPVARIANT var;
+		PropVariantInit(&var);
+
+		HRESULT Result = Src->GetItem(Key, &var);
+		
+		if (SUCCEEDED(Result))
+		{
+			Result = Dest->SetItem(Key, var);
+			PropVariantClear(&var);
+		}
+
+		return Result;
+	}
+
+
 	TComPtr<IMFMediaType> CreateOutputType(const GUID& MajorType, const GUID& SubType, bool AllowNonStandardCodecs)
 	{
 		TComPtr<IMFMediaType> OutputType;
@@ -261,7 +349,7 @@ namespace WmfMedia
 
 			if (FAILED(Result))
 			{
-				UE_LOG(LogWmfMedia, Warning, TEXT("Failed to create %s output type (%s)"), *MajorTypeToString(MajorType), *ResultToString(Result));
+				UE_LOG(LogWmfMedia, Warning, TEXT("Failed to create %s output type: %s"), *MajorTypeToString(MajorType), *ResultToString(Result));
 				return NULL;
 			}
 
@@ -269,7 +357,7 @@ namespace WmfMedia
 
 			if (FAILED(Result))
 			{
-				UE_LOG(LogWmfMedia, Warning, TEXT("Failed to initialize %s output type (%s)"), *MajorTypeToString(MajorType), *ResultToString(Result));
+				UE_LOG(LogWmfMedia, Warning, TEXT("Failed to initialize %s output type: %s"), *MajorTypeToString(MajorType), *ResultToString(Result));
 				return NULL;
 			}
 		}
@@ -334,7 +422,7 @@ namespace WmfMedia
 
 			if (FAILED(Result))
 			{
-				UE_LOG(LogWmfMedia, Warning, TEXT("Failed to initialize binary output type (%s)"), *ResultToString(Result));
+				UE_LOG(LogWmfMedia, Warning, TEXT("Failed to initialize binary output type: %s"), *ResultToString(Result));
 				return NULL;
 			}
 		}
@@ -345,7 +433,7 @@ namespace WmfMedia
 
 			if (FAILED(Result))
 			{
-				UE_LOG(LogWmfMedia, Warning, TEXT("Failed to initialize caption output type (%s)"), *ResultToString(Result));
+				UE_LOG(LogWmfMedia, Warning, TEXT("Failed to initialize caption output type: %s"), *ResultToString(Result));
 				return NULL;
 			}
 		}
@@ -383,7 +471,7 @@ namespace WmfMedia
 
 			if (FAILED(Result))
 			{
-				UE_LOG(LogWmfMedia, Warning, TEXT("Failed to set video output type (%s)"), *ResultToString(Result));
+				UE_LOG(LogWmfMedia, Warning, TEXT("Failed to set video output type: %s"), *ResultToString(Result));
 				return NULL;
 			}
 
@@ -405,7 +493,7 @@ namespace WmfMedia
 
 			if (FAILED(Result))
 			{
-				UE_LOG(LogWmfMedia, Warning, TEXT("Failed to set video output sub-type (%s)"), *ResultToString(Result));
+				UE_LOG(LogWmfMedia, Warning, TEXT("Failed to set video output sub-type: %s"), *ResultToString(Result));
 				return NULL;
 			}
 		}
@@ -418,6 +506,139 @@ namespace WmfMedia
 	}
 
 
+	FString DumpAttributes(IMFAttributes& Attributes)
+	{
+		FString Dump;
+
+		UINT32 NumAttributes = 0;
+		{
+			const HRESULT Result = Attributes.GetCount(&NumAttributes);
+
+			if (FAILED(Result))
+			{
+				Dump += FString::Printf(TEXT("\tFailed to get attribute count: %s\n"), *ResultToString(Result));
+				return Dump;
+			}
+		}
+
+		for (UINT32 AttributeIndex = 0; AttributeIndex < NumAttributes; ++AttributeIndex)
+		{
+			GUID Guid = { 0 };
+			PROPVARIANT Item;
+			{
+				::PropVariantInit(&Item);
+
+				const HRESULT Result = Attributes.GetItemByIndex(AttributeIndex, &Guid, &Item);
+
+				if (FAILED(Result))
+				{
+					Dump += FString::Printf(TEXT("\tFailed to get attribute %i: %s\n"), AttributeIndex, *ResultToString(Result));
+				}
+			}
+
+			const FString GuidName = AttributeToString(Guid);
+
+			if (Guid == MF_MT_AM_FORMAT_TYPE)
+			{
+				Dump += Dump += FString::Printf(TEXT("\t%s: %s (%s)\n"), *GuidName, *GuidToString(*Item.puuid), *FormatTypeToString(*Item.puuid));
+			}
+			else if (Guid == MF_MT_MAJOR_TYPE)
+			{
+				Dump += FString::Printf(TEXT("\t%s: %s (%s)\n"), *GuidName, *GuidToString(*Item.puuid), *MajorTypeToString(*Item.puuid));
+			}
+			else if (Guid == MF_MT_SUBTYPE)
+			{
+				Dump += FString::Printf(TEXT("\t%s: %s (%s)\n"), *GuidName, *GuidToString(*Item.puuid), *SubTypeToString(*Item.puuid));
+			}
+			else if ((Guid == MF_MT_FRAME_RATE) || (Guid == MF_MT_FRAME_RATE_RANGE_MAX) || (Guid == MF_MT_FRAME_RATE_RANGE_MIN))
+			{
+				UINT32 High = 0;
+				UINT32 Low = 0;
+
+				::Unpack2UINT32AsUINT64(Item.uhVal.QuadPart, &High, &Low);
+				Dump += FString::Printf(TEXT("\t%s: %d/%d\n"), *GuidName, High, Low);
+			}
+			else if (Guid == MF_MT_FRAME_SIZE)
+			{
+				UINT32 High = 0;
+				UINT32 Low = 0;
+
+				::Unpack2UINT32AsUINT64(Item.uhVal.QuadPart, &High, &Low);
+				Dump += FString::Printf(TEXT("\t%s: %d x %d\n"), *GuidName, High, Low);
+			}
+			else if (Guid == MF_MT_INTERLACE_MODE)
+			{
+				Dump += FString::Printf(TEXT("\t%s: %d (%s)\n"), *GuidName, Item.ulVal, *InterlaceModeToString((MFVideoInterlaceMode)Item.ulVal));
+			}
+			else if (Guid == MF_MT_PIXEL_ASPECT_RATIO)
+			{
+				UINT32 High = 0;
+				UINT32 Low = 0;
+
+				::Unpack2UINT32AsUINT64(Item.uhVal.QuadPart, &High, &Low);
+				Dump += FString::Printf(TEXT("\t%s: %d:%d\n"), *GuidName, High, Low);
+			}
+			else if ((Guid == MF_MT_GEOMETRIC_APERTURE) || (Guid == MF_MT_MINIMUM_DISPLAY_APERTURE) || (Guid == MF_MT_PAN_SCAN_APERTURE))
+			{
+				if (Item.caub.cElems < sizeof(MFVideoArea))
+				{
+					Dump += FString::Printf(TEXT("\t%s: failed to get value (buffer too small)\n"), *GuidName);
+				}
+				else
+				{
+					MFVideoArea* Area = (MFVideoArea*)Item.caub.pElems;
+
+					Dump += FString::Printf(TEXT("\t%s: (%f,%f) (%d,%d)\n"),
+						*GuidName,
+						Area->OffsetX.value + (static_cast<float>(Area->OffsetX.fract) / 65536.0f),
+						Area->OffsetY.value + (static_cast<float>(Area->OffsetY.fract) / 65536.0f),
+						Area->Area.cx,
+						Area->Area.cy);
+				}
+			}
+			else
+			{
+				switch (Item.vt)
+				{
+				case VT_UI4:
+					Dump += FString::Printf(TEXT("\t%s: %d\n"), *GuidName, Item.ulVal);
+					break;
+
+				case VT_UI8:
+					Dump += FString::Printf(TEXT("\t%s: %ll\n"), *GuidName, Item.uhVal.QuadPart);
+					break;
+
+				case VT_R8:
+					Dump += FString::Printf(TEXT("\t%s: %f\n"), *GuidName, Item.dblVal);
+					break;
+
+				case VT_CLSID:
+					Dump += FString::Printf(TEXT("\t%s: %s\n"), *GuidName, *GuidToString(*Item.puuid));
+					break;
+
+				case VT_LPWSTR:
+					Dump += FString::Printf(TEXT("\t%s: %s\n"), *GuidName, Item.pwszVal);
+					break;
+
+				case VT_VECTOR | VT_UI1:
+					Dump += FString::Printf(TEXT("\t%s: <byte array>\n"), *GuidName);
+					break;
+
+				case VT_UNKNOWN:
+					Dump += FString::Printf(TEXT("\t%s: IUnknown\n"), *GuidName);
+					break;
+
+				default:
+					Dump += FString::Printf(TEXT("\t%s: Unknown value type %d\n"), *GuidName, Item.vt);
+					break;
+				}
+			}
+		}
+
+		return Dump;
+	}
+
+
 	void EnumerateCaptureDevices(GUID DeviceType, TArray<TComPtr<IMFActivate>>& OutDevices)
 	{
 		if ((DeviceType != MF_DEVSOURCE_ATTRIBUTE_SOURCE_TYPE_AUDCAP_GUID) &&
@@ -433,7 +654,7 @@ namespace WmfMedia
 
 			if (FAILED(Result))
 			{
-				UE_LOG(LogWmfMedia, Error, TEXT("Failed to create capture device enumeration attributes: %s"), *WmfMedia::ResultToString(Result));
+				UE_LOG(LogWmfMedia, Error, TEXT("Failed to create capture device enumeration attributes: %s"), *ResultToString(Result));
 				return;
 			}
 		}
@@ -444,7 +665,7 @@ namespace WmfMedia
 
 			if (FAILED(Result))
 			{
-				UE_LOG(LogWmfMedia, Error, TEXT("Failed to set capture device enumeration type: %s"), *WmfMedia::ResultToString(Result));
+				UE_LOG(LogWmfMedia, Error, TEXT("Failed to set capture device enumeration type: %s"), *ResultToString(Result));
 				return;
 			}
 		}
@@ -486,6 +707,23 @@ namespace WmfMedia
 	}
 
 
+	FString FormatTypeToString(GUID FormatType)
+	{
+		if (FormatType == FORMAT_DvInfo) return TEXT("DVINFO");
+		if (FormatType == FORMAT_MPEG2Video) return TEXT("MPEG2VIDEOINFO");
+		if (FormatType == FORMAT_MPEGStreams) return TEXT("AM_MPEGSYSTEMTYPE");
+		if (FormatType == FORMAT_MPEGVideo) return TEXT("MPEG1VIDEOINFO");
+		if (FormatType == FORMAT_None) return TEXT("None");
+		if (FormatType == FORMAT_VideoInfo) return TEXT("VIDEOINFOHEADER");
+		if (FormatType == FORMAT_VideoInfo2) return TEXT("VIDEOINFOHEADER2");
+		if (FormatType == FORMAT_WaveFormatEx) return TEXT("WAVEFORMATEX");
+		if (FormatType == FORMAT_525WSS) return TEXT("525WSS");
+		if (FormatType == GUID_NULL) return TEXT("Null");
+
+		return FString::Printf(TEXT("Unknown format type %s"), *GuidToString(FormatType));
+	}
+
+
 	bool FrameRateToRatio(float FrameRate, int32& OutNumerator, int32& OutDenominator)
 	{
 		if (FrameRate < 0.0f)
@@ -630,6 +868,33 @@ namespace WmfMedia
 	}
 
 
+	HRESULT GetTopologyFromEvent(IMFMediaEvent& Event, TComPtr<IMFTopology>& OutTopology)
+	{
+		HRESULT Result = S_OK;
+		PROPVARIANT Variant;
+
+		PropVariantInit(&Variant);
+		Result = Event.GetValue(&Variant);
+
+		if (SUCCEEDED(Result))
+		{
+			if (Variant.vt != VT_UNKNOWN)
+			{
+				Result = E_UNEXPECTED;
+			}
+		}
+
+		if (SUCCEEDED(Result))
+		{
+			Result = Variant.punkVal->QueryInterface(IID_PPV_ARGS(&OutTopology));
+		}
+
+		PropVariantClear(&Variant);
+
+		return Result;
+	}
+
+
 	FString GuidToString(const GUID& Guid)
 	{
 		return FString::Printf(TEXT("%08x-%04x-%04x-%02x%02x%02x%02x%02x%02x%02x%02x"),
@@ -648,6 +913,24 @@ namespace WmfMedia
 	}
 
 
+	FString InterlaceModeToString(MFVideoInterlaceMode Mode)
+	{
+		switch (Mode)
+		{
+		case MFVideoInterlace_Unknown: return TEXT("Unknown");
+		case MFVideoInterlace_Progressive: return TEXT("Progressive");
+		case MFVideoInterlace_FieldInterleavedUpperFirst: return TEXT("Field Interleaved Upper First");
+		case MFVideoInterlace_FieldInterleavedLowerFirst: return TEXT("Field Intereaved Lower First");
+		case MFVideoInterlace_FieldSingleUpper: return TEXT("Field Single Upper");
+		case MFVideoInterlace_FieldSingleLower: return TEXT("Field Single Lower");
+		case MFVideoInterlace_MixedInterlaceOrProgressive: return TEXT("Mixed Interlace or Progressive");
+		
+		default:
+			return FString::Printf(TEXT("Unknown mode %i"), (int32)Mode);
+		}
+	}
+
+
 	bool IsSupportedMajorType(const GUID& MajorType)
 	{
 		for (GUID const* const Type : SupportedMajorTypes)
@@ -680,6 +963,21 @@ namespace WmfMedia
 	}
 
 
+	FString MarkerTypeToString(MFSTREAMSINK_MARKER_TYPE MarkerType)
+	{
+		switch (MarkerType)
+		{
+		case MFSTREAMSINK_MARKER_DEFAULT: return TEXT("Default");
+		case MFSTREAMSINK_MARKER_ENDOFSEGMENT: return TEXT("End Of Segment");
+		case MFSTREAMSINK_MARKER_TICK: return TEXT("Tick");
+		case MFSTREAMSINK_MARKER_EVENT: return TEXT("Event");
+
+		default:
+			return FString::Printf(TEXT("Unknown marker type %i"), (int32)MarkerType);
+		}
+	}
+
+
 	FString MediaEventToString(MediaEventType Event)
 	{
 		switch (Event)
@@ -796,7 +1094,7 @@ namespace WmfMedia
 		case MEContentProtectionMetadata: return TEXT("Content Protection Metadata");
 
 		default:
-			return FString::Printf(TEXT("Unknown event %i"), Event);
+			return FString::Printf(TEXT("Unknown event %i"), (int32)Event);
 		}
 	}
 
@@ -840,7 +1138,7 @@ namespace WmfMedia
 
 					if (FAILED(Result))
 					{
-						UE_LOG(LogWmfMedia, Error, TEXT("Failed to create capture device attributes (%s)"), *WmfMedia::ResultToString(Result));
+						UE_LOG(LogWmfMedia, Error, TEXT("Failed to create capture device attributes: %s"), *ResultToString(Result));
 						return NULL;
 					}
 
@@ -852,7 +1150,7 @@ namespace WmfMedia
 
 					if (FAILED(Result))
 					{
-						UE_LOG(LogWmfMedia, Error, TEXT("Failed to set capture device source type attribute (%s)"), *WmfMedia::ResultToString(Result));
+						UE_LOG(LogWmfMedia, Error, TEXT("Failed to set capture device source type attribute: %s"), *ResultToString(Result));
 						return NULL;
 					}
 
@@ -867,7 +1165,7 @@ namespace WmfMedia
 
 					if (FAILED(Result))
 					{
-						UE_LOG(LogWmfMedia, Error, TEXT("Failed to set capture device endpoint/symlink attribute (%s)"), *WmfMedia::ResultToString(Result));
+						UE_LOG(LogWmfMedia, Error, TEXT("Failed to set capture device endpoint/symlink attribute: %s"), *ResultToString(Result));
 						return NULL;
 					}
 				}
@@ -878,7 +1176,7 @@ namespace WmfMedia
 
 					if (FAILED(Result))
 					{
-						UE_LOG(LogWmfMedia, Error, TEXT("Failed to create capture device media source (%s)"), *WmfMedia::ResultToString(Result));
+						UE_LOG(LogWmfMedia, Error, TEXT("Failed to create capture device media source: %s"), *ResultToString(Result));
 						return NULL;
 					}
 				}
@@ -930,7 +1228,7 @@ namespace WmfMedia
 
 			if (FAILED(Result))
 			{
-				UE_LOG(LogWmfMedia, Error, TEXT("Failed to create media source resolver (%s)"), *WmfMedia::ResultToString(Result));
+				UE_LOG(LogWmfMedia, Error, TEXT("Failed to create media source resolver: %s"), *ResultToString(Result));
 				return NULL;
 			}
 		}
@@ -947,7 +1245,7 @@ namespace WmfMedia
 
 				if (FAILED(Result))
 				{
-					UE_LOG(LogWmfMedia, Error, TEXT("Failed to resolve byte stream %s (%s)"), *Url, *WmfMedia::ResultToString(Result));
+					UE_LOG(LogWmfMedia, Error, TEXT("Failed to resolve byte stream %s: %s"), *Url, *ResultToString(Result));
 					return NULL;
 				}
 			}
@@ -957,7 +1255,7 @@ namespace WmfMedia
 
 				if (FAILED(Result))
 				{
-					UE_LOG(LogWmfMedia, Error, TEXT("Failed to resolve URL %s (%s)"), *Url, *WmfMedia::ResultToString(Result));
+					UE_LOG(LogWmfMedia, Error, TEXT("Failed to resolve URL %s: %s"), *Url, *ResultToString(Result));
 					return NULL;
 				}
 			}
@@ -970,7 +1268,7 @@ namespace WmfMedia
 
 			if (FAILED(Result))
 			{
-				UE_LOG(LogWmfMedia, Error, TEXT("Failed to query media source interface: %s"), *WmfMedia::ResultToString(Result));
+				UE_LOG(LogWmfMedia, Error, TEXT("Failed to query media source interface: %s"), *ResultToString(Result));
 				return NULL;
 			}
 		}
@@ -1419,6 +1717,26 @@ namespace WmfMedia
 		// unknown type
 		return FString::Printf(TEXT("%s (%s)"), *GuidToString(SubType), *FourccToString(SubType.Data1));
 	}
+
+
+	FString TopologyStatusToString(MF_TOPOSTATUS Status)
+	{
+		switch (Status)
+		{
+		case MF_TOPOSTATUS_ENDED: return TEXT("Ended");
+		case MF_TOPOSTATUS_INVALID: return TEXT("Invalid");
+		case MF_TOPOSTATUS_READY: return TEXT("Ready");
+		case MF_TOPOSTATUS_SINK_SWITCHED: return TEXT("Sink Switched");
+		case MF_TOPOSTATUS_STARTED_SOURCE: return TEXT("Started Source");
+
+#if (WINVER >= _WIN32_WINNT_WIN7)
+		case MF_TOPOSTATUS_DYNAMIC_CHANGED: return TEXT("Dynamic Changed");
+#endif
+
+		default:
+			return FString::Printf(TEXT("Unknown status %i"), (int32)Status);
+		}
+	}
 }
 
 
diff --git a/Engine/Plugins/Media/WmfMedia/Source/WmfMedia/Private/Wmf/WmfMediaUtils.h b/Engine/Plugins/Media/WmfMedia/Source/WmfMedia/Private/Wmf/WmfMediaUtils.h
--- a/Engine/Plugins/Media/WmfMedia/Source/WmfMedia/Private/Wmf/WmfMediaUtils.h
+++ b/Engine/Plugins/Media/WmfMedia/Source/WmfMedia/Private/Wmf/WmfMediaUtils.h
@@ -14,6 +14,33 @@
 
 namespace WmfMedia
 {
+	/**
+	 * Convert a media attribute identifier to string.
+	 *
+	 * @param Guid The unique identifier of the attribute.
+	 * @return The corresponding name string.
+	 */
+	FString AttributeToString(const GUID& Guid);
+
+	/**
+	 * Convert a capture device role to string.
+	 *
+	 * @param Role The device role to convert.
+	 * @return The corresponding string.
+	 * @see CaptureDeviceRoleToString, GuidToString, MajorTypeToString, MarkerTypeToString, MediaEventToString, ResultToString, SubTypeToString, TopologyStatusToString
+	 */
+	FString CaptureDeviceRoleToString(ERole Role);
+
+	/**
+	 * Copy an attribute from one attribute collection to another.
+	 *
+	 * @param Src The attribute collection to copy from.
+	 * @param Dest The attribute collection to copy to.
+	 * @param Key The key of the attribute to copy.
+	 * @return Result code.
+	 */
+	HRESULT CopyAttribute(IMFAttributes* Src, IMFAttributes* Dest, const GUID& Key);
+
 	/**
 	 * Create an output media type for the given input media type.
 	 *
@@ -24,6 +51,13 @@ namespace WmfMedia
 	 */
 	TComPtr<IMFMediaType> CreateOutputType(const GUID& MajorType, const GUID& SubType, bool AllowNonStandardCodecs);
 
+	/**
+	 * Dump information about the given media type to the output log.
+	 *
+	 * @param MediaType The media type.
+	 */
+	FString DumpAttributes(IMFAttributes& Attributes);
+
 	/**
 	 * Enumerate the available capture devices.
 	 *
@@ -35,6 +69,14 @@ namespace WmfMedia
 	 */
 	void EnumerateCaptureDevices(GUID DeviceType, TArray<TComPtr<IMFActivate>>& OutDevices);
 
+	/**
+	 * Convert an MF_MT_AM_FORMAT_TYPE value to string.
+	 *
+	 * @param FormatType The unique identifier of the format type.
+	 * @return The corresponding string.
+	 */
+	FString FormatTypeToString(GUID FormatType);
+
 	/**
 	 * Convert a floating point frame rate to an integer ratio.
 	 *
@@ -50,6 +92,7 @@ namespace WmfMedia
 	 *
 	 * @param Fourcc The code to convert.
 	 * @return The corresponding string.
+	 * @see CaptureDeviceRoleToString, GuidToString, MajorTypeToString, MarkerTypeToString, MediaEventToString, ResultToString, SubTypeToString, TopologyStatusToString
 	 */
 	FString FourccToString(unsigned long Fourcc);
 
@@ -74,14 +117,32 @@ namespace WmfMedia
 	 */
 	const TArray<TComPtr<IMFMediaType>>& GetSupportedMediaTypes(const GUID& MajorType);
 
+	/**
+	 * Get the playback topology object from the given media event.
+	 *
+	 * @param Event The event object.
+	 * @param OutTopology Will contain the topology.
+	 * @return Result code.
+	 */
+	HRESULT GetTopologyFromEvent(IMFMediaEvent& Event, TComPtr<IMFTopology>& OutTopology);
+
 	/**
 	 * Convert a Windows GUID to string.
 	 *
 	 * @param Guid The GUID to convert.
 	 * @return The corresponding string.
+	 * @see CaptureDeviceRoleToString, FourccToString, MajorTypeToString, MarkerTypeToString, MediaEventToString, ResultToString, SubTypeToString, TopologyStatusToString
 	 */
 	FString GuidToString(const GUID& Guid);
 
+	/**
+	 * Convert a video interlace mode to string.
+	 *
+	 * @param Mode The interlace mode to convert.
+	 * @return The corresponding string.
+	 */
+	FString InterlaceModeToString(MFVideoInterlaceMode Mode);
+
 	/**
 	 * Check whether the given major media type is supported.
 	 *
@@ -96,14 +157,25 @@ namespace WmfMedia
 	 *
 	 * @param MajorType The major type GUID to convert.
 	 * @return The corresponding string.
+	 * @see CaptureDeviceRoleToString, FourccToString, GuidToString, MarkerTypeToString, MediaEventToString, ResultToString, SubTypeToString, TopologyStatusToString
 	 */
 	FString MajorTypeToString(const GUID& MajorType);
 
+	/**
+	 * Convert a stream sink marker type to string.
+	 *
+	 * @param MarkerType The marker type to convert.
+	 * @return The corresponding string.
+	 * @see CaptureDeviceRoleToString, FourccToString, GuidToString, MajorTypeToString, MediaEventToString, ResultToString, SubTypeToString, TopologyStatusToString
+	 */
+	FString MarkerTypeToString(MFSTREAMSINK_MARKER_TYPE MarkerType);
+
 	/**
 	 * Convert a media event to string.
 	 *
 	 * @param Event The event code to convert.
 	 * @return The corresponding string.
+	 * @see CaptureDeviceRoleToString, FourccToString, GuidToString, MajorTypeToString, MarkerTypeToString, ResultToString, SubTypeToString, TopologyStatusToString
 	 */
 	FString MediaEventToString(MediaEventType Event);
 
@@ -132,6 +204,7 @@ namespace WmfMedia
 	 *
 	 * @param Result The result code to convert.
 	 * @return The corresponding string.
+	 * @see CaptureDeviceRoleToString, FourccToString, GuidToString, MajorTypeToString, MarkerTypeToString, MediaEventToString, SubTypeToString, TopologyStatusToString
 	 */
 	FString ResultToString(HRESULT Result);
 
@@ -140,8 +213,18 @@ namespace WmfMedia
 	 *
 	 * @param SubType The sub-type GUID to convert.
 	 * @return The corresponding string.
+	 * @see CaptureDeviceRoleToString, FourccToString, GuidToString, MajorTypeToString, MarkerTypeToString, MediaEventToString, ResultToString, TopologyStatusToString
 	 */
 	FString SubTypeToString(const GUID& SubType);
+
+	/**
+	 * Convert a WMF topology status to string.
+	 *
+	 * @param Status The status to convert.
+	 * @return The corresponding string.
+	 * @see CaptureDeviceRoleToString, FourccToString, GuidToString, MajorTypeToString, MarkerTypeToString, MediaEventToString, ResultToString, SubTypeToString
+	 */
+	FString TopologyStatusToString(MF_TOPOSTATUS Status);
 }
 
 
diff --git a/Engine/Plugins/Media/WmfMedia/Source/WmfMedia/Private/WmfMediaModule.cpp b/Engine/Plugins/Media/WmfMedia/Source/WmfMedia/Private/WmfMediaModule.cpp
--- a/Engine/Plugins/Media/WmfMedia/Source/WmfMedia/Private/WmfMediaModule.cpp
+++ b/Engine/Plugins/Media/WmfMedia/Source/WmfMedia/Private/WmfMediaModule.cpp
@@ -2,13 +2,16 @@
 
 #include "WmfMediaPrivate.h"
 
+#include "IMediaCaptureSupport.h"
+#include "Modules/ModuleManager.h"
+
+#include "IWmfMediaModule.h"
+
 #if WMFMEDIA_SUPPORTED_PLATFORM
-	#include "IMediaCaptureSupport.h"
 	#include "IMediaModule.h"
 	#include "Modules/ModuleInterface.h"
-	#include "Modules/ModuleManager.h"
+	#include "Templates/SharedPointer.h"
 
-	#include "IWmfMediaModule.h"
 	#include "WmfMediaPlayer.h"
 	#include "WmfMediaUtils.h"
 
diff --git a/Engine/Plugins/Media/WmfMedia/Source/WmfMedia/Private/WmfMediaPrivate.h b/Engine/Plugins/Media/WmfMedia/Source/WmfMedia/Private/WmfMediaPrivate.h
--- a/Engine/Plugins/Media/WmfMedia/Source/WmfMedia/Private/WmfMediaPrivate.h
+++ b/Engine/Plugins/Media/WmfMedia/Source/WmfMedia/Private/WmfMediaPrivate.h
@@ -4,7 +4,7 @@
 
 #include "Logging/LogMacros.h"
 
-#define WMFMEDIA_SUPPORTED_PLATFORM (PLATFORM_WINDOWS && (WINVER >= 0x0600 /*Vista*/))
+#define WMFMEDIA_SUPPORTED_PLATFORM (PLATFORM_WINDOWS && (WINVER >= 0x0600 /*Vista*/) && !UE_SERVER)
 
 #if WMFMEDIA_SUPPORTED_PLATFORM
 	#include "../../WmfMediaFactory/Public/WmfMediaSettings.h"
@@ -23,8 +23,16 @@
 	#include <nserror.h>
 	#include <shlwapi.h>
 
+	const GUID FORMAT_525WSS = { 0xc7ecf04d, 0x4582, 0x4869, { 0x9a, 0xbb, 0xbf, 0xb5, 0x23, 0xb6, 0x2e, 0xdf } };
+	const GUID FORMAT_DvInfo = { 0x05589f84, 0xc356, 0x11ce, { 0xbf, 0x01, 0x00, 0xaa, 0x00, 0x55, 0x59, 0x5a } };
+	const GUID FORMAT_MPEG2Video = { 0xe06d80e3, 0xdb46, 0x11cf, { 0xb4, 0xd1, 0x00, 0x80, 0x05f, 0x6c, 0xbb, 0xea } };
+	const GUID FORMAT_MPEGStreams = { 0x05589f83, 0xc356, 0x11ce, { 0xbf, 0x01, 0x00, 0xaa, 0x00, 0x55, 0x59, 0x5a } };
+	const GUID FORMAT_MPEGVideo = { 0x05589f82, 0xc356, 0x11ce, { 0xbf, 0x01, 0x00, 0xaa, 0x00, 0x55, 0x59, 0x5a } };
+	const GUID FORMAT_None = { 0x0F6417D6, 0xc318, 0x11d0, { 0xa4, 0x3f, 0x00, 0xa0, 0xc9, 0x22, 0x31, 0x96 } };
 	const GUID FORMAT_VideoInfo = { 0x05589f80, 0xc356, 0x11ce, { 0xbf, 0x01, 0x00, 0xaa, 0x00, 0x55, 0x59, 0x5a } };
 	const GUID FORMAT_VideoInfo2 = { 0xf72a76A0, 0xeb0a, 0x11d0, { 0xac, 0xe4, 0x00, 0x00, 0xc0, 0xcc, 0x16, 0xba } };
+	const GUID FORMAT_WaveFormatEx = { 0x05589f81, 0xc356, 0x11ce,{ 0xbf, 0x01, 0x00, 0xaa, 0x00, 0x55, 0x59, 0x5a } };
+	const GUID MR_VIDEO_ACCELERATION_SERVICE = { 0xefef5175, 0x5c7d, 0x4ce2, { 0xbb, 0xbd, 0x34, 0xff, 0x8b, 0xca, 0x65, 0x54 } };
 
 	#if (WINVER < _WIN32_WINNT_WIN8)
 		const GUID MF_LOW_LATENCY = { 0x9c27891a, 0xed7a, 0x40e1, { 0x88, 0xe8, 0xb2, 0x27, 0x27, 0xa0, 0x24, 0xee } };
diff --git a/Engine/Plugins/MovieScene/ActorSequence/Source/ActorSequence/ActorSequence.Build.cs b/Engine/Plugins/MovieScene/ActorSequence/Source/ActorSequence/ActorSequence.Build.cs
new file mode 100644
--- /dev/null
+++ b/Engine/Plugins/MovieScene/ActorSequence/Source/ActorSequence/ActorSequence.Build.cs
@@ -0,0 +1,20 @@
+// Copyright 1998-2016 Epic Games, Inc. All Rights Reserved.
+
+using UnrealBuildTool;
+
+public class ActorSequence : ModuleRules
+{
+	public ActorSequence(ReadOnlyTargetRules Target) : base(Target)
+	{
+		PrivateIncludePaths.Add("ActorSequence/Private");
+
+		PublicDependencyModuleNames.AddRange(
+			new string[] {
+				"Core",
+				"CoreUObject",
+				"Engine",
+				"MovieScene",
+			}
+		);
+	}
+}
diff --git a/Engine/Plugins/MovieScene/ActorSequence/Source/ActorSequence/Private/ActorSequencePrivatePCH.h b/Engine/Plugins/MovieScene/ActorSequence/Source/ActorSequence/Private/ActorSequencePrivatePCH.h
new file mode 100644
--- /dev/null
+++ b/Engine/Plugins/MovieScene/ActorSequence/Source/ActorSequence/Private/ActorSequencePrivatePCH.h
@@ -0,0 +1,11 @@
+// Copyright 1998-2016 Epic Games, Inc. All Rights Reserved.
+
+#pragma once
+
+
+/* Private dependencies
+ *****************************************************************************/
+
+#include "Core.h"
+#include "CoreUObject.h"
+#include "Engine.h"
diff --git a/Engine/Plugins/MovieScene/ActorSequence/Source/ActorSequence/Public/ActorSequence.h b/Engine/Plugins/MovieScene/ActorSequence/Source/ActorSequence/Public/ActorSequence.h
new file mode 100644
--- /dev/null
+++ b/Engine/Plugins/MovieScene/ActorSequence/Source/ActorSequence/Public/ActorSequence.h
@@ -0,0 +1,69 @@
+// Copyright 1998-2016 Epic Games, Inc. All Rights Reserved.
+
+#pragma once
+
+#include "MovieSceneSequence.h"
+#include "MovieScene.h"
+#include "ActorSequenceObjectReference.h"
+#include "ActorSequence.generated.h"
+
+/**
+ * Movie scene animation embedded within an actor.
+ */
+UCLASS(BlueprintType, Experimental, DefaultToInstanced)
+class ACTORSEQUENCE_API UActorSequence
+	: public UMovieSceneSequence
+{
+public:
+	GENERATED_BODY()
+
+	UActorSequence(const FObjectInitializer& ObjectInitializer);
+
+	//~ UMovieSceneSequence interface
+	virtual void BindPossessableObject(const FGuid& ObjectId, UObject& PossessedObject, UObject* Context) override;
+	virtual bool CanPossessObject(UObject& Object, UObject* InPlaybackContext) const override;
+	virtual void LocateBoundObjects(const FGuid& ObjectId, UObject* Context, TArray<UObject*, TInlineAllocator<1>>& OutObjects) const override;
+	virtual UMovieScene* GetMovieScene() const override;
+	virtual UObject* GetParentObject(UObject* Object) const override;
+	virtual void UnbindPossessableObjects(const FGuid& ObjectId) override;
+
+#if WITH_EDITOR
+	virtual FText GetDisplayName() const override;
+#endif
+
+	UBlueprint* GetParentBlueprint() const;
+
+	bool IsEditable() const;
+	
+private:
+
+	//~ UObject interface
+	virtual void PostInitProperties() override;
+
+private:
+	
+	/** Pointer to the movie scene that controls this animation. */
+	UPROPERTY(Instanced)
+	UMovieScene* MovieScene;
+
+	/** Collection of object references. */
+	UPROPERTY()
+	FActorSequenceObjectReferenceMap ObjectReferences;
+
+#if WITH_EDITOR
+public:
+
+	/** Event that is fired to initialize default state for a sequence */
+	DECLARE_EVENT_OneParam(UActorSequence, FOnInitialize, UActorSequence*)
+	static FOnInitialize& OnInitializeSequence() { return OnInitializeSequenceEvent; }
+
+private:
+	static FOnInitialize OnInitializeSequenceEvent;
+#endif
+
+#if WITH_EDITORONLY_DATA
+private:
+	UPROPERTY()
+	bool bHasBeenInitialized;
+#endif
+};
\ No newline at end of file
diff --git a/Engine/Plugins/MovieScene/ActorSequence/Source/ActorSequence/Public/ActorSequenceComponent.h b/Engine/Plugins/MovieScene/ActorSequence/Source/ActorSequence/Public/ActorSequenceComponent.h
new file mode 100644
--- /dev/null
+++ b/Engine/Plugins/MovieScene/ActorSequence/Source/ActorSequence/Public/ActorSequenceComponent.h
@@ -0,0 +1,54 @@
+// Copyright 1998-2016 Epic Games, Inc. All Rights Reserved.
+
+#pragma once
+
+#include "Components/ActorComponent.h"
+#include "MovieSceneSequencePlayer.h"
+#include "ActorSequenceComponent.generated.h"
+
+
+class UActorSequence;
+class UActorSequencePlayer;
+
+
+/**
+ * Movie scene animation embedded within an actor.
+ */
+UCLASS(Blueprintable, Experimental, ClassGroup=Sequence, hidecategories=(Collision, Cooking, Activation), meta=(BlueprintSpawnableComponent))
+class ACTORSEQUENCE_API UActorSequenceComponent
+	: public UActorComponent
+{
+public:
+	GENERATED_BODY()
+
+	UActorSequenceComponent(const FObjectInitializer& Init);
+
+	UActorSequence* GetSequence() const
+	{
+		return Sequence;
+	}
+
+	UActorSequencePlayer* GetSequencePlayer() const 
+	{
+		return SequencePlayer;
+	}
+	
+	virtual void PostInitProperties() override;
+	virtual void BeginPlay() override;
+	virtual void TickComponent(float DeltaTime, enum ELevelTick TickType, FActorComponentTickFunction *ThisTickFunction) override;
+
+protected:
+
+	UPROPERTY(EditAnywhere, Category="Playback", meta=(ShowOnlyInnerProperties))
+	FMovieSceneSequencePlaybackSettings PlaybackSettings;
+
+	/** Embedded actor sequence data */
+	UPROPERTY(EditAnywhere, Instanced, Category=Animation)
+	UActorSequence* Sequence;
+
+	UPROPERTY(transient, BlueprintReadOnly, Category=Animation)
+	UActorSequencePlayer* SequencePlayer;
+
+	UPROPERTY(EditAnywhere, Category="Playback")
+	bool bAutoPlay;
+};
diff --git a/Engine/Plugins/MovieScene/ActorSequence/Source/ActorSequence/Public/ActorSequencePlayer.h b/Engine/Plugins/MovieScene/ActorSequence/Source/ActorSequence/Public/ActorSequencePlayer.h
new file mode 100644
--- /dev/null
+++ b/Engine/Plugins/MovieScene/ActorSequence/Source/ActorSequence/Public/ActorSequencePlayer.h
@@ -0,0 +1,26 @@
+// Copyright 1998-2016 Epic Games, Inc. All Rights Reserved.
+
+#pragma once
+
+#include "IMovieScenePlayer.h"
+#include "ActorSequence.h"
+#include "MovieSceneSequencePlayer.h"
+#include "ActorSequencePlayer.generated.h"
+
+/**
+ * UActorSequencePlayer is used to actually "play" an actor sequence asset at runtime.
+ */
+UCLASS(BlueprintType)
+class ACTORSEQUENCE_API UActorSequencePlayer
+	: public UMovieSceneSequencePlayer
+{
+public:
+	GENERATED_BODY()
+
+protected:
+
+	//~ IMovieScenePlayer interface
+	virtual UObject* GetPlaybackContext() const override;
+	virtual TArray<UObject*> GetEventContexts() const override;
+};
+
diff --git a/Engine/Plugins/MovieScene/ActorSequenceEditor/ActorSequenceEditor.uplugin b/Engine/Plugins/MovieScene/ActorSequenceEditor/ActorSequenceEditor.uplugin
--- a/Engine/Plugins/MovieScene/ActorSequenceEditor/ActorSequenceEditor.uplugin
+++ b/Engine/Plugins/MovieScene/ActorSequenceEditor/ActorSequenceEditor.uplugin
@@ -14,6 +14,13 @@
 	"CanContainContent" : false,
 	"IsBetaVersion" : true,
 	"Installed" : false,
+	"Plugins" :
+	[
+		{
+      "Name": "ActorSequence",
+      "Enabled" :  true
+		}
+	],
 	"Modules" :
 	[
 		{
diff --git a/Engine/Plugins/MovieScene/LevelSequenceEditor/Source/LevelSequenceEditor/Private/CinematicViewport/SCinematicTransportRange.cpp b/Engine/Plugins/MovieScene/LevelSequenceEditor/Source/LevelSequenceEditor/Private/CinematicViewport/SCinematicTransportRange.cpp
--- a/Engine/Plugins/MovieScene/LevelSequenceEditor/Source/LevelSequenceEditor/Private/CinematicViewport/SCinematicTransportRange.cpp
+++ b/Engine/Plugins/MovieScene/LevelSequenceEditor/Source/LevelSequenceEditor/Private/CinematicViewport/SCinematicTransportRange.cpp
@@ -8,6 +8,7 @@
 #include "Styles/LevelSequenceEditorStyle.h"
 #include "MovieScene.h"
 #include "ISequencer.h"
+#include "CommonMovieSceneTools.h"
 
 #define LOCTEXT_NAMESPACE "SCinematicTransportRange"
 
@@ -95,7 +96,13 @@ void SCinematicTransportRange::Tick(const FGeometry& AllottedGeometry, const dou
 	ISequencer* Sequencer = GetSequencer();
 	if (Sequencer)
 	{
-		Sequencer->GetKeysFromSelection(ActiveKeyCollection);
+		TRange<float> WorkingRange = Sequencer->GetFocusedMovieSceneSequence()->GetMovieScene()->GetEditorData().WorkingRange;
+
+		// Anything within 3 pixel's worth of time is a duplicate as far as we're concerned
+		FTimeToPixel TimeToPixelConverter(AllottedGeometry, WorkingRange);
+		const float DuplicateThreshold = (TimeToPixelConverter.PixelToTime(3.f) - TimeToPixelConverter.PixelToTime(0.f));
+
+		Sequencer->GetKeysFromSelection(ActiveKeyCollection, DuplicateThreshold);
 	}
 }
 
diff --git a/Engine/Plugins/MovieScene/LevelSequenceEditor/Source/LevelSequenceEditor/Private/Misc/LevelSequenceEditorActorBinding.cpp b/Engine/Plugins/MovieScene/LevelSequenceEditor/Source/LevelSequenceEditor/Private/Misc/LevelSequenceEditorActorBinding.cpp
--- a/Engine/Plugins/MovieScene/LevelSequenceEditor/Source/LevelSequenceEditor/Private/Misc/LevelSequenceEditorActorBinding.cpp
+++ b/Engine/Plugins/MovieScene/LevelSequenceEditor/Source/LevelSequenceEditor/Private/Misc/LevelSequenceEditorActorBinding.cpp
@@ -39,16 +39,17 @@ bool FLevelSequenceEditorActorBinding::SupportsSequence(UMovieSceneSequence* InS
 
 void FLevelSequenceEditorActorBinding::AddPossessActorMenuExtensions(FMenuBuilder& MenuBuilder)
 {
-	auto IsActorValidForPossession = [=](const AActor* InActor, TSharedRef<ISequencer> InSequencer)
+	auto IsActorValidForPossession = [=](const AActor* InActor, TWeakPtr<ISequencer> InWeakSequencer)
 	{
+		TSharedPtr<ISequencer> SequencerPtr = InWeakSequencer.Pin();
 		bool bCreateHandleIfMissing = false;
-		return !InSequencer->GetHandleToObject((UObject*)InActor, bCreateHandleIfMissing).IsValid();
+		return SequencerPtr.IsValid() && !SequencerPtr->GetHandleToObject((UObject*)InActor, bCreateHandleIfMissing).IsValid();
 	};
 
 	// Set up a menu entry to add the selected actor(s) to the sequencer
 	TArray<AActor*> ActorsValidForPossession;
 	GEditor->GetSelectedActors()->GetSelectedObjects(ActorsValidForPossession);
-	ActorsValidForPossession.RemoveAll([&](AActor* In){ return !IsActorValidForPossession(In, Sequencer.Pin().ToSharedRef()); });
+	ActorsValidForPossession.RemoveAll([&](AActor* In){ return !IsActorValidForPossession(In, Sequencer); });
 
 	FText SelectedLabel;
 	FSlateIcon ActorIcon = FSlateIconFinder::FindIconForClass(AActor::StaticClass());
@@ -90,7 +91,7 @@ void FLevelSequenceEditorActorBinding::AddPossessActorMenuExtensions(FMenuBuilde
 		InitOptions.ColumnMap.Add(FBuiltInColumnTypes::Label(), FColumnInfo(EColumnVisibility::Visible, 0));
 
 		// Only display actors that are not possessed already
-		InitOptions.Filters->AddFilterPredicate(FActorFilterPredicate::CreateLambda(IsActorValidForPossession, Sequencer.Pin().ToSharedRef()));
+		InitOptions.Filters->AddFilterPredicate(FActorFilterPredicate::CreateLambda(IsActorValidForPossession, Sequencer));
 	}
 
 	// actor selector to allow the user to choose an actor
diff --git a/Engine/Plugins/MovieScene/LevelSequenceEditor/Source/LevelSequenceEditor/Private/Misc/LevelSequenceEditorActorSpawner.cpp b/Engine/Plugins/MovieScene/LevelSequenceEditor/Source/LevelSequenceEditor/Private/Misc/LevelSequenceEditorActorSpawner.cpp
--- a/Engine/Plugins/MovieScene/LevelSequenceEditor/Source/LevelSequenceEditor/Private/Misc/LevelSequenceEditorActorSpawner.cpp
+++ b/Engine/Plugins/MovieScene/LevelSequenceEditor/Source/LevelSequenceEditor/Private/Misc/LevelSequenceEditorActorSpawner.cpp
@@ -47,8 +47,10 @@ TValueOrError<FNewSpawnable, FText> FLevelSequenceEditorActorSpawner::CreateNewS
 	// Deal with creating a spawnable from an instance of an actor
 	else if (AActor* Actor = Cast<AActor>(&SourceObject))
 	{
-		NewSpawnable.ObjectTemplate = StaticDuplicateObject(Actor, &OwnerMovieScene, TemplateName, RF_AllFlags & ~RF_Transactional);
-		NewSpawnable.Name = Actor->GetActorLabel();
+		AActor* SpawnedActor = Cast<AActor>(StaticDuplicateObject(Actor, &OwnerMovieScene, TemplateName, RF_AllFlags & ~RF_Transactional));
+		SpawnedActor->bIsEditorPreviewActor = false;
+		NewSpawnable.ObjectTemplate = SpawnedActor;
+		NewSpawnable.Name = Actor->GetActorLabel();		
 	}
 
 	// If it's a blueprint, we need some special handling
diff --git a/Engine/Source/Editor/MovieSceneTools/Private/CurveKeyAreas/FloatCurveKeyArea.cpp b/Engine/Source/Editor/MovieSceneTools/Private/CurveKeyAreas/FloatCurveKeyArea.cpp
--- a/Engine/Source/Editor/MovieSceneTools/Private/CurveKeyAreas/FloatCurveKeyArea.cpp
+++ b/Engine/Source/Editor/MovieSceneTools/Private/CurveKeyAreas/FloatCurveKeyArea.cpp
@@ -179,6 +179,7 @@ FRichCurve* FFloatCurveKeyArea::GetRichCurve()
 TArray<FKeyHandle> FFloatCurveKeyArea::GetUnsortedKeyHandles() const
 {
 	TArray<FKeyHandle> OutKeyHandles;
+	OutKeyHandles.Reserve(Curve->GetNumKeys());
 
 	for (auto It(Curve->GetKeyHandleIterator()); It; ++It)
 	{
diff --git a/Engine/Source/Editor/MovieSceneTools/Private/CurveKeyAreas/IntegralKeyArea.cpp b/Engine/Source/Editor/MovieSceneTools/Private/CurveKeyAreas/IntegralKeyArea.cpp
--- a/Engine/Source/Editor/MovieSceneTools/Private/CurveKeyAreas/IntegralKeyArea.cpp
+++ b/Engine/Source/Editor/MovieSceneTools/Private/CurveKeyAreas/IntegralKeyArea.cpp
@@ -95,12 +95,13 @@ float FIntegralCurveKeyAreaBase::GetKeyTime(FKeyHandle KeyHandle) const
 FRichCurve* FIntegralCurveKeyAreaBase::GetRichCurve()
 {
 	return nullptr;
-};
+}
 
 
 TArray<FKeyHandle> FIntegralCurveKeyAreaBase::GetUnsortedKeyHandles() const
 {
 	TArray<FKeyHandle> OutKeyHandles;
+	OutKeyHandles.Reserve(Curve.GetNumKeys());
 
 	for (auto It(Curve.GetKeyHandleIterator()); It; ++It)
 	{
diff --git a/Engine/Source/Editor/MovieSceneTools/Private/CurveKeyAreas/NameCurveKeyArea.cpp b/Engine/Source/Editor/MovieSceneTools/Private/CurveKeyAreas/NameCurveKeyArea.cpp
--- a/Engine/Source/Editor/MovieSceneTools/Private/CurveKeyAreas/NameCurveKeyArea.cpp
+++ b/Engine/Source/Editor/MovieSceneTools/Private/CurveKeyAreas/NameCurveKeyArea.cpp
@@ -115,6 +115,7 @@ FRichCurve* FNameCurveKeyArea::GetRichCurve()
 TArray<FKeyHandle> FNameCurveKeyArea::GetUnsortedKeyHandles() const
 {
 	TArray<FKeyHandle> OutKeyHandles;
+	OutKeyHandles.Reserve(Curve.GetNumKeys());
 
 	for (auto It(Curve.GetKeyHandleIterator()); It; ++It)
 	{
diff --git a/Engine/Source/Editor/MovieSceneTools/Private/CurveKeyAreas/StringCurveKeyArea.cpp b/Engine/Source/Editor/MovieSceneTools/Private/CurveKeyAreas/StringCurveKeyArea.cpp
--- a/Engine/Source/Editor/MovieSceneTools/Private/CurveKeyAreas/StringCurveKeyArea.cpp
+++ b/Engine/Source/Editor/MovieSceneTools/Private/CurveKeyAreas/StringCurveKeyArea.cpp
@@ -144,6 +144,7 @@ FRichCurve* FStringCurveKeyArea::GetRichCurve()
 TArray<FKeyHandle> FStringCurveKeyArea::GetUnsortedKeyHandles() const
 {
 	TArray<FKeyHandle> OutKeyHandles;
+	OutKeyHandles.Reserve(Curve->GetNumKeys());
 
 	for (auto It(Curve->GetKeyHandleIterator()); It; ++It)
 	{
diff --git a/Engine/Source/Editor/MovieSceneTools/Private/MovieSceneBuiltInEasingFunctionCustomization.cpp b/Engine/Source/Editor/MovieSceneTools/Private/MovieSceneBuiltInEasingFunctionCustomization.cpp
--- a/Engine/Source/Editor/MovieSceneTools/Private/MovieSceneBuiltInEasingFunctionCustomization.cpp
+++ b/Engine/Source/Editor/MovieSceneTools/Private/MovieSceneBuiltInEasingFunctionCustomization.cpp
@@ -12,6 +12,7 @@
 #include "SButton.h"
 #include "SBox.h"
 #include "SlateApplication.h"
+#include "ScopedTransaction.h"
 
 #include "EditorStyleSet.h"
 
@@ -226,6 +227,8 @@ void FMovieSceneBuiltInEasingFunctionCustomization::CustomizeDetails(IDetailLayo
 
 FReply FMovieSceneBuiltInEasingFunctionCustomization::SetType(EMovieSceneBuiltInEasing NewType)
 {
+	FScopedTransaction Transaction(NSLOCTEXT("EasingFunctionCustomization", "SetEasingType", "Set Easing Type"));	
+
 	TypeProperty->NotifyPreChange();
 
 	TArray<void*> RawData;
diff --git a/Engine/Source/Editor/MovieSceneTools/Private/Sections/CinematicShotSection.cpp b/Engine/Source/Editor/MovieSceneTools/Private/Sections/CinematicShotSection.cpp
--- a/Engine/Source/Editor/MovieSceneTools/Private/Sections/CinematicShotSection.cpp
+++ b/Engine/Source/Editor/MovieSceneTools/Private/Sections/CinematicShotSection.cpp
@@ -87,6 +87,24 @@ void FCinematicShotSection::ResizeSection(ESequencerSectionResizeMode ResizeMode
 	FViewportThumbnailSection::ResizeSection(ResizeMode, ResizeTime);
 }
 
+void FCinematicShotSection::BeginSlipSection()
+{
+	BeginResizeSection();
+}
+
+void FCinematicShotSection::SlipSection(float SlipTime)
+{
+	float StartOffset = (SlipTime - InitialStartTimeDuringResize) / SectionObject.Parameters.TimeScale;
+	StartOffset += InitialStartOffsetDuringResize;
+
+	// Ensure start offset is not less than 0
+	StartOffset = FMath::Max(StartOffset, 0.f);
+
+	SectionObject.Parameters.StartOffset = StartOffset;
+
+	FViewportThumbnailSection::SlipSection(SlipTime);
+}
+
 void FCinematicShotSection::Tick(const FGeometry& AllottedGeometry, const FGeometry& ClippedGeometry, const double InCurrentTime, const float InDeltaTime)
 {
 	// Set cached data
diff --git a/Engine/Source/Editor/MovieSceneTools/Private/Sections/CinematicShotSection.h b/Engine/Source/Editor/MovieSceneTools/Private/Sections/CinematicShotSection.h
--- a/Engine/Source/Editor/MovieSceneTools/Private/Sections/CinematicShotSection.h
+++ b/Engine/Source/Editor/MovieSceneTools/Private/Sections/CinematicShotSection.h
@@ -39,6 +39,8 @@ public:
 	virtual FMargin GetContentPadding() const override;
 	virtual void BeginResizeSection() override;
 	virtual void ResizeSection(ESequencerSectionResizeMode ResizeMode, float ResizeTime) override;
+	virtual void BeginSlipSection() override;
+	virtual void SlipSection(float SlipTime) override;
 
 	// FThumbnail interface
 	virtual void SetSingleTime(float GlobalTime) override;
diff --git a/Engine/Source/Editor/MovieSceneTools/Private/TrackEditors/AudioTrackEditor.cpp b/Engine/Source/Editor/MovieSceneTools/Private/TrackEditors/AudioTrackEditor.cpp
--- a/Engine/Source/Editor/MovieSceneTools/Private/TrackEditors/AudioTrackEditor.cpp
+++ b/Engine/Source/Editor/MovieSceneTools/Private/TrackEditors/AudioTrackEditor.cpp
@@ -613,6 +613,8 @@ FAudioSection::FAudioSection( UMovieSceneSection& InSection, TWeakPtr<ISequencer
 	, StoredDrawRange(TRange<float>::Empty())
 	, StoredSectionHeight(0.f)
 	, Sequencer(InSequencer)
+	, InitialStartOffsetDuringResize(0.f)
+	, InitialStartTimeDuringResize(0.f)
 {
 }
 
@@ -734,6 +736,27 @@ void FAudioSection::Tick( const FGeometry& AllottedGeometry, const FGeometry& Pa
 	}
 }
 
+void FAudioSection::BeginSlipSection()
+{
+	UMovieSceneAudioSection* AudioSection = Cast<UMovieSceneAudioSection>(&Section);
+	InitialStartOffsetDuringResize = AudioSection->GetStartOffset();
+	InitialStartTimeDuringResize = AudioSection->GetStartTime();
+}
+
+void FAudioSection::SlipSection(float SlipTime)
+{
+	UMovieSceneAudioSection* AudioSection = Cast<UMovieSceneAudioSection>(&Section);
+
+	float StartOffset = (SlipTime - InitialStartTimeDuringResize);
+	StartOffset += InitialStartOffsetDuringResize;
+
+	// Ensure start offset is not less than 0
+	StartOffset = FMath::Max(StartOffset, 0.f);
+
+	AudioSection->SetStartOffset(StartOffset);
+
+	ISequencerSection::SlipSection(SlipTime);
+}
 
 void FAudioSection::RegenerateWaveforms(TRange<float> DrawRange, int32 XOffset, int32 XSize, const FColor& ColorTint, float DisplayScale)
 {
diff --git a/Engine/Source/Editor/MovieSceneTools/Private/TrackEditors/PropertyTrackEditors/TransformPropertyTrackEditor.cpp b/Engine/Source/Editor/MovieSceneTools/Private/TrackEditors/PropertyTrackEditors/TransformPropertyTrackEditor.cpp
--- a/Engine/Source/Editor/MovieSceneTools/Private/TrackEditors/PropertyTrackEditors/TransformPropertyTrackEditor.cpp
+++ b/Engine/Source/Editor/MovieSceneTools/Private/TrackEditors/PropertyTrackEditors/TransformPropertyTrackEditor.cpp
@@ -26,13 +26,19 @@ TSharedRef<ISequencerSection> FTransformPropertyTrackEditor::MakeSectionInterfac
 
 TSharedPtr<SWidget> FTransformPropertyTrackEditor::BuildOutlinerEditWidget(const FGuid& ObjectBinding, UMovieSceneTrack* Track, const FBuildEditWidgetParams& Params)
 {
-	TSharedPtr<ISequencer> SequencerPtr = GetSequencer();
+	TWeakPtr<ISequencer> WeakSequencer = GetSequencer();
 
 	const int32 RowIndex = Params.TrackInsertRowIndex;
 	auto SubMenuCallback = [=]() -> TSharedRef<SWidget>
 	{
 		FMenuBuilder MenuBuilder(true, nullptr);
-		FSequencerUtilities::PopulateMenu_CreateNewSection(MenuBuilder, RowIndex, Track, SequencerPtr);
+
+		TSharedPtr<ISequencer> SequencerPtr = WeakSequencer.Pin();
+		if (SequencerPtr.IsValid())
+		{
+			FSequencerUtilities::PopulateMenu_CreateNewSection(MenuBuilder, RowIndex, Track, SequencerPtr);
+		}
+
 		return MenuBuilder.MakeWidget();
 	};
 
diff --git a/Engine/Source/Editor/MovieSceneTools/Private/TrackEditors/SkeletalAnimationTrackEditor.cpp b/Engine/Source/Editor/MovieSceneTools/Private/TrackEditors/SkeletalAnimationTrackEditor.cpp
--- a/Engine/Source/Editor/MovieSceneTools/Private/TrackEditors/SkeletalAnimationTrackEditor.cpp
+++ b/Engine/Source/Editor/MovieSceneTools/Private/TrackEditors/SkeletalAnimationTrackEditor.cpp
@@ -146,6 +146,27 @@ void FSkeletalAnimationSection::ResizeSection(ESequencerSectionResizeMode Resize
 	ISequencerSection::ResizeSection(ResizeMode, ResizeTime);
 }
 
+void FSkeletalAnimationSection::BeginSlipSection()
+{
+	BeginResizeSection();
+}
+
+void FSkeletalAnimationSection::SlipSection(float SlipTime)
+{
+	float StartOffset = (SlipTime - InitialStartTimeDuringResize) * Section.Params.PlayRate;
+	StartOffset += InitialStartOffsetDuringResize;
+
+	// Ensure start offset is not less than 0
+	if (StartOffset < 0)
+	{
+		StartOffset = 0.f;
+	}
+
+	Section.Params.StartOffset = StartOffset;
+
+	ISequencerSection::SlipSection(SlipTime);
+}
+
 FSkeletalAnimationTrackEditor::FSkeletalAnimationTrackEditor( TSharedRef<ISequencer> InSequencer )
 	: FMovieSceneTrackEditor( InSequencer ) 
 { }
diff --git a/Engine/Source/Editor/MovieSceneTools/Private/TrackEditors/SkeletalAnimationTrackEditor.h b/Engine/Source/Editor/MovieSceneTools/Private/TrackEditors/SkeletalAnimationTrackEditor.h
--- a/Engine/Source/Editor/MovieSceneTools/Private/TrackEditors/SkeletalAnimationTrackEditor.h
+++ b/Engine/Source/Editor/MovieSceneTools/Private/TrackEditors/SkeletalAnimationTrackEditor.h
@@ -99,6 +99,8 @@ public:
 	virtual int32 OnPaintSection( FSequencerSectionPainter& Painter ) const override;
 	virtual void BeginResizeSection() override;
 	virtual void ResizeSection(ESequencerSectionResizeMode ResizeMode, float ResizeTime) override;
+	virtual void BeginSlipSection() override;
+	virtual void SlipSection(float SlipTime) override;
 
 private:
 
diff --git a/Engine/Source/Editor/MovieSceneTools/Private/TrackEditors/SubTrackEditor.cpp b/Engine/Source/Editor/MovieSceneTools/Private/TrackEditors/SubTrackEditor.cpp
--- a/Engine/Source/Editor/MovieSceneTools/Private/TrackEditors/SubTrackEditor.cpp
+++ b/Engine/Source/Editor/MovieSceneTools/Private/TrackEditors/SubTrackEditor.cpp
@@ -45,6 +45,8 @@ public:
 		: DisplayName(InDisplayName)
 		, SectionObject(*CastChecked<UMovieSceneSubSection>(&InSection))
 		, Sequencer(InSequencer)
+		, InitialStartOffsetDuringResize(0.f)
+		, InitialStartTimeDuringResize(0.f)
 	{
 	}
 
@@ -278,6 +280,25 @@ public:
 		return FReply::Handled();
 	}
 
+	virtual void BeginSlipSection() override
+	{
+		InitialStartOffsetDuringResize = SectionObject.Parameters.StartOffset;
+		InitialStartTimeDuringResize = SectionObject.GetStartTime();
+	}
+
+	virtual void SlipSection(float SlipTime) override
+	{
+		float StartOffset = (SlipTime - InitialStartTimeDuringResize) / SectionObject.Parameters.TimeScale;
+		StartOffset += InitialStartOffsetDuringResize;
+
+		// Ensure start offset is not less than 0
+		StartOffset = FMath::Max(StartOffset, 0.f);
+
+		SectionObject.Parameters.StartOffset = StartOffset;
+
+		ISequencerSection::SlipSection(SlipTime);
+	}
+
 private:
 
 	/** Display name of the section */
@@ -288,6 +309,12 @@ private:
 
 	/** Sequencer interface */
 	TWeakPtr<ISequencer> Sequencer;
+
+	/** Cached start offset value valid only during resize */
+	float InitialStartOffsetDuringResize;
+	
+	/** Cached start time valid only during resize */
+	float InitialStartTimeDuringResize;
 };
 
 
diff --git a/Engine/Source/Editor/MovieSceneTools/Public/TrackEditors/AudioTrackEditor.h b/Engine/Source/Editor/MovieSceneTools/Public/TrackEditors/AudioTrackEditor.h
--- a/Engine/Source/Editor/MovieSceneTools/Public/TrackEditors/AudioTrackEditor.h
+++ b/Engine/Source/Editor/MovieSceneTools/Public/TrackEditors/AudioTrackEditor.h
@@ -106,6 +106,8 @@ public:
 	virtual void GenerateSectionLayout( class ISectionLayoutBuilder& LayoutBuilder) const override;
 	virtual int32 OnPaintSection(FSequencerSectionPainter& Painter) const override;
 	virtual void Tick(const FGeometry& AllottedGeometry, const FGeometry& ParentGeometry, const double InCurrentTime, const float InDeltaTime) override;
+	virtual void BeginSlipSection() override;
+	virtual void SlipSection(float SlipTime) override;
 	
 private:
 
@@ -136,4 +138,10 @@ private:
 	TWeakObjectPtr<USoundWave> StoredSoundWave;
 
 	TWeakPtr<ISequencer> Sequencer;
+
+	/** Cached start offset value valid only during resize */
+	float InitialStartOffsetDuringResize;
+	
+	/** Cached start time valid only during resize */
+	float InitialStartTimeDuringResize;
 };
diff --git a/Engine/Source/Editor/SequenceRecorder/Private/SSequenceRecorder.cpp b/Engine/Source/Editor/SequenceRecorder/Private/SSequenceRecorder.cpp
--- a/Engine/Source/Editor/SequenceRecorder/Private/SSequenceRecorder.cpp
+++ b/Engine/Source/Editor/SequenceRecorder/Private/SSequenceRecorder.cpp
@@ -16,9 +16,11 @@
 #include "AnimationRecorder.h"
 #include "PropertyEditorModule.h"
 #include "IDetailsView.h"
+#include "DragAndDrop/ActorDragDropOp.h"
 #include "SequenceRecorderCommands.h"
 #include "SequenceRecorderSettings.h"
 #include "SequenceRecorder.h"
+#include "SDropTarget.h"
 
 #define LOCTEXT_NAMESPACE "SequenceRecorder"
 
@@ -180,24 +182,30 @@ void SSequenceRecorder::Construct(const FArguments& Args)
 						+SVerticalBox::Slot()
 						.FillHeight(1.0f)
 						[
-							SAssignNew(ListView, SListView<UActorRecording*>)
-							.ListItemsSource(&FSequenceRecorder::Get().GetQueuedRecordings())
-							.SelectionMode(ESelectionMode::SingleToggle)
-							.OnGenerateRow(this, &SSequenceRecorder::MakeListViewWidget)
-							.OnSelectionChanged(this, &SSequenceRecorder::OnSelectionChanged)
-							.HeaderRow
-							(
-								SNew(SHeaderRow)
-								+ SHeaderRow::Column(ActorColumnName)
-								.FillWidth(43.0f)
-								.DefaultLabel(LOCTEXT("ActorHeaderName", "Actor"))
-								+ SHeaderRow::Column(AnimationColumnName)
-								.FillWidth(43.0f)
-								.DefaultLabel(LOCTEXT("AnimationHeaderName", "Animation"))
-								+ SHeaderRow::Column(LengthColumnName)
-								.FillWidth(14.0f)
-								.DefaultLabel(LOCTEXT("LengthHeaderName", "Length"))
-							)
+							SNew( SDropTarget )
+							.OnAllowDrop( this, &SSequenceRecorder::OnRecordingListAllowDrop )
+							.OnDrop( this, &SSequenceRecorder::OnRecordingListDrop )
+							.Content()
+							[
+								SAssignNew(ListView, SListView<UActorRecording*>)
+								.ListItemsSource(&FSequenceRecorder::Get().GetQueuedRecordings())
+								.SelectionMode(ESelectionMode::SingleToggle)
+								.OnGenerateRow(this, &SSequenceRecorder::MakeListViewWidget)
+								.OnSelectionChanged(this, &SSequenceRecorder::OnSelectionChanged)
+								.HeaderRow
+								(
+									SNew(SHeaderRow)
+									+ SHeaderRow::Column(ActorColumnName)
+									.FillWidth(43.0f)
+									.DefaultLabel(LOCTEXT("ActorHeaderName", "Actor"))
+									+ SHeaderRow::Column(AnimationColumnName)
+									.FillWidth(43.0f)
+									.DefaultLabel(LOCTEXT("AnimationHeaderName", "Animation"))
+									+ SHeaderRow::Column(LengthColumnName)
+									.FillWidth(14.0f)
+									.DefaultLabel(LOCTEXT("LengthHeaderName", "Length"))
+								)
+							]
 						]
 						+SVerticalBox::Slot()
 						.AutoHeight()
@@ -403,4 +411,30 @@ FText SSequenceRecorder::GetTargetSequenceName() const
 	return FText::Format(LOCTEXT("NextSequenceFormat", "Next Sequence: {0}"), FText::FromString(FSequenceRecorder::Get().GetNextSequenceName()));
 }
 
+bool SSequenceRecorder::OnRecordingListAllowDrop( TSharedPtr<FDragDropOperation> DragDropOperation )
+{
+	return DragDropOperation->IsOfType<FActorDragDropOp>();
+}
+
+
+FReply SSequenceRecorder::OnRecordingListDrop( TSharedPtr<FDragDropOperation> DragDropOperation )
+{
+	if ( DragDropOperation->IsOfType<FActorDragDropOp>() )
+	{
+		TSharedPtr<FActorDragDropOp> ActorDragDropOperation = StaticCastSharedPtr<FActorDragDropOp>( DragDropOperation );
+
+		for (auto Actor : ActorDragDropOperation->Actors)
+		{
+			if (Actor.IsValid())
+			{
+				FSequenceRecorder::Get().AddNewQueuedRecording(Actor.Get());
+			}
+		}
+
+		return FReply::Handled();
+	}
+
+	return FReply::Unhandled();
+}
+
 #undef LOCTEXT_NAMESPACE
diff --git a/Engine/Source/Editor/SequenceRecorder/Private/SSequenceRecorder.h b/Engine/Source/Editor/SequenceRecorder/Private/SSequenceRecorder.h
--- a/Engine/Source/Editor/SequenceRecorder/Private/SSequenceRecorder.h
+++ b/Engine/Source/Editor/SequenceRecorder/Private/SSequenceRecorder.h
@@ -14,6 +14,7 @@ class FActiveTimerHandle;
 class FUICommandList;
 class IDetailsView;
 class SProgressBar;
+class FDragDropOperation;
 
 class SSequenceRecorder : public SCompoundWidget
 {
@@ -71,6 +72,10 @@ private:
 
 	FText GetTargetSequenceName() const;
 
+	FReply OnRecordingListDrop( TSharedPtr<FDragDropOperation> DragDropOperation );
+
+	bool OnRecordingListAllowDrop( TSharedPtr<FDragDropOperation> DragDropOperation );
+
 private:
 	TSharedPtr<IDetailsView> SequenceRecordingDetailsView;
 
diff --git a/Engine/Source/Editor/SequenceRecorder/SequenceRecorder.Build.cs b/Engine/Source/Editor/SequenceRecorder/SequenceRecorder.Build.cs
--- a/Engine/Source/Editor/SequenceRecorder/SequenceRecorder.Build.cs
+++ b/Engine/Source/Editor/SequenceRecorder/SequenceRecorder.Build.cs
@@ -46,6 +46,7 @@ namespace UnrealBuildTool.Rules
 					"NetworkReplayStreaming",
 					"AssetRegistry",
 					"CinematicCamera",
+                    "EditorWidgets",
 				}
 				);
 
diff --git a/Engine/Source/Editor/Sequencer/Private/DisplayNodes/SequencerDisplayNode.cpp b/Engine/Source/Editor/Sequencer/Private/DisplayNodes/SequencerDisplayNode.cpp
--- a/Engine/Source/Editor/Sequencer/Private/DisplayNodes/SequencerDisplayNode.cpp
+++ b/Engine/Source/Editor/Sequencer/Private/DisplayNodes/SequencerDisplayNode.cpp
@@ -24,6 +24,7 @@
 #include "CommonMovieSceneTools.h"
 #include "Framework/Commands/GenericCommands.h"
 #include "ScopedTransaction.h"
+#include "SequencerKeyTimeCache.h"
 
 #define LOCTEXT_NAMESPACE "SequencerDisplayNode"
 
@@ -35,8 +36,30 @@ namespace SequencerNodeConstants
 {
 	extern const float CommonPadding;
 	const float CommonPadding = 4.f;
+
+	static const FVector2D KeyMarkSize = FVector2D(3.f, 21.f);
 }
 
+struct FNameAndSignature
+{
+	FGuid Signature;
+	FName Name;
+
+	bool IsValid() const
+	{
+		return Signature.IsValid() && !Name.IsNone();
+	}
+
+	friend bool operator==(const FNameAndSignature& A, const FNameAndSignature& B)
+	{
+		return A.Signature == B.Signature && A.Name == B.Name;
+	}
+
+	friend uint32 GetTypeHash(const FNameAndSignature& In)
+	{
+		return HashCombine(GetTypeHash(In.Signature), GetTypeHash(In.Name));
+	}
+};
 
 class SSequencerObjectTrack
 	: public SLeafWidget
@@ -50,7 +73,8 @@ public:
 
 	/** SLeafWidget Interface */
 	virtual int32 OnPaint( const FPaintArgs& Args, const FGeometry& AllottedGeometry, const FSlateRect& MyCullingRect, FSlateWindowElementList& OutDrawElements, int32 LayerId, const FWidgetStyle& InWidgetStyle, bool bParentEnabled ) const override;
-	
+	virtual void Tick( const FGeometry& AllottedGeometry, const double InCurrentTime, const float InDeltaTime ) override;
+
 	void Construct( const FArguments& InArgs, TSharedRef<FSequencerDisplayNode> InRootNode )
 	{
 		RootNode = InRootNode;
@@ -68,10 +92,7 @@ protected:
 private:
 
 	/** Collects all key times from the root node */
-	void CollectAllKeyTimes(TArray<float>& OutKeyTimes) const;
-
-	/** Adds a key time uniquely to an array of key times */
-	void AddKeyTime(const float& NewTime, TArray<float>& OutKeyTimes) const;
+	void GenerateCachedKeyPositions(const FGeometry& AllottedGeometry);
 
 private:
 
@@ -80,27 +101,47 @@ private:
 
 	/** The current view range */
 	TAttribute< TRange<float> > ViewRange;
+
+	FSequencerKeyCollectionSignature KeyCollectionSignature;
+
+	/** The time-range for which KeyDrawPositions was generated */
+	TRange<float> CachedViewRange;
+	/** Cached pixel positions for all keys in the current view range */
+	TArray<float> KeyDrawPositions;
+	/** Cached key times per key area. Updated when section signature changes */
+	TMap<FNameAndSignature, FSequencerCachedKeys> SectionToKeyTimeCache;
 };
 
 
+void SSequencerObjectTrack::Tick( const FGeometry& AllottedGeometry, const double InCurrentTime, const float InDeltaTime )
+{
+	SWidget::Tick(AllottedGeometry, InCurrentTime, InDeltaTime);
+
+	FSequencerKeyCollectionSignature NewCollectionSignature = FSequencerKeyCollectionSignature::FromNodesRecursive({ RootNode.Get() }, 0.f);
+	if (NewCollectionSignature != KeyCollectionSignature || CachedViewRange != ViewRange.Get())
+	{
+		CachedViewRange = ViewRange.Get();
+		KeyCollectionSignature = MoveTemp(NewCollectionSignature);
+		GenerateCachedKeyPositions(AllottedGeometry);
+	}
+}
+
 int32 SSequencerObjectTrack::OnPaint(const FPaintArgs& Args, const FGeometry& AllottedGeometry, const FSlateRect& MyCullingRect, FSlateWindowElementList& OutDrawElements, int32 LayerId, const FWidgetStyle& InWidgetStyle, bool bParentEnabled) const
 {
 	if (RootNode->GetSequencer().GetSettings()->GetShowCombinedKeyframes())
 	{
-		TArray<float> OutKeyTimes;
-		CollectAllKeyTimes(OutKeyTimes);
-	
-		FTimeToPixel TimeToPixelConverter(AllottedGeometry, ViewRange.Get());
-
-		for (int32 i = 0; i < OutKeyTimes.Num(); ++i)
+		for (float KeyPosition : KeyDrawPositions)
 		{
-			float KeyPosition = TimeToPixelConverter.TimeToPixel(OutKeyTimes[i]);
-			static const FVector2D KeyMarkSize = FVector2D(3.f, 21.f);
-
 			FSlateDrawElement::MakeBox(
 				OutDrawElements,
 				LayerId+1,
-				AllottedGeometry.ToPaintGeometry(FVector2D(KeyPosition - FMath::CeilToFloat(KeyMarkSize.X/2.f), FMath::CeilToFloat(AllottedGeometry.GetLocalSize().Y/2.f - KeyMarkSize.Y/2.f)), KeyMarkSize),
+				AllottedGeometry.ToPaintGeometry(
+					FVector2D(
+						KeyPosition - FMath::CeilToFloat(SequencerNodeConstants::KeyMarkSize.X/2.f),
+						FMath::CeilToFloat(AllottedGeometry.GetLocalSize().Y/2.f - SequencerNodeConstants::KeyMarkSize.Y/2.f)
+					),
+					SequencerNodeConstants::KeyMarkSize
+				),
 				FEditorStyle::GetBrush("Sequencer.KeyMark"),
 				ESlateDrawEffect::None,
 				FLinearColor(1.f, 1.f, 1.f, 1.f)
@@ -120,38 +161,102 @@ FVector2D SSequencerObjectTrack::ComputeDesiredSize( float ) const
 }
 
 
-void SSequencerObjectTrack::CollectAllKeyTimes(TArray<float>& OutKeyTimes) const
+void SSequencerObjectTrack::GenerateCachedKeyPositions(const FGeometry& AllottedGeometry)
 {
-	TArray<TSharedRef<FSequencerSectionKeyAreaNode>> OutNodes;
-	RootNode->GetChildKeyAreaNodesRecursively(OutNodes);
+	static float DuplicateThresholdPx = 3.f;
+
+	// Swap the last frame's cache with a temporary so we start this frame's cache from a clean slate
+	TMap<FNameAndSignature, FSequencerCachedKeys> PreviouslyCachedKeyTimes;
+	Swap(PreviouslyCachedKeyTimes, SectionToKeyTimeCache);
 
-for (int32 i = 0; i < OutNodes.Num(); ++i)
+	// Unnamed key areas are uncacheable, so we track those separately
+	TArray<FSequencerCachedKeys> UncachableKeyTimes;
+
+	// First off, accumulate (and cache) KeyDrawPositions as times, we convert to positions in the later loop
+	for (auto& CachePair : KeyCollectionSignature.GetKeyAreas())
 	{
-		TArray< TSharedRef<IKeyArea> > KeyAreas = OutNodes[i]->GetAllKeyAreas();
-		for (int32 j = 0; j < KeyAreas.Num(); ++j)
+		TSharedRef<IKeyArea> KeyArea = CachePair.Key;
+
+		UMovieSceneSection* Section = KeyArea->GetOwningSection();
+		FNameAndSignature CacheKey{ CachePair.Value, KeyArea->GetName() };
+
+		// If we cached this last frame, use those key times again
+		FSequencerCachedKeys* CachedKeyTimes = CacheKey.IsValid() ? PreviouslyCachedKeyTimes.Find(CacheKey) : nullptr;
+		if (CachedKeyTimes)
 		{
-			TArray<FKeyHandle> KeyHandles = KeyAreas[j]->GetUnsortedKeyHandles();
-			for (int32 k = 0; k < KeyHandles.Num(); ++k)
-			{
-				AddKeyTime(KeyAreas[j]->GetKeyTime(KeyHandles[k]), OutKeyTimes);
-			}
+			SectionToKeyTimeCache.Add(CacheKey, MoveTemp(*CachedKeyTimes));
+			continue;
+		}
+
+		// Generate a new cache
+		FSequencerCachedKeys TempCache;
+		TempCache.Update(KeyArea);
+
+		if (CacheKey.IsValid())
+		{
+			SectionToKeyTimeCache.Add(CacheKey, MoveTemp(TempCache));
+		}
+		else
+		{
+			UncachableKeyTimes.Add(MoveTemp(TempCache));
 		}
 	}
-}
 
+	KeyDrawPositions.Reset();
 
-void SSequencerObjectTrack::AddKeyTime(const float& NewTime, TArray<float>& OutKeyTimes) const
-{
-	// @todo Sequencer It might be more efficient to add each key and do the pruning at the end
-	for (float& KeyTime : OutKeyTimes)
+	// Instead of accumulating all key times into a single array and then sorting (which doesn't scale well with large numbers),
+	// we use a collection of iterators that are only incremented when they've been added to the KeyDrawPositions array
+	struct FIter
+	{
+		FIter(TArrayView<const FSequencerCachedKey> In) : KeysInRange(In), CurrentIndex(0) {}
+
+		explicit                     operator bool() const   { return KeysInRange.IsValidIndex(CurrentIndex); }
+		FIter&                       operator++()            { ++CurrentIndex; return *this; }
+
+		const FSequencerCachedKey&   operator*() const       { return KeysInRange[CurrentIndex]; }
+		const FSequencerCachedKey*   operator->() const      { return &KeysInRange[CurrentIndex]; }
+	private:
+		TArrayView<const FSequencerCachedKey> KeysInRange;
+		int32 CurrentIndex;
+	};
+
+	TArray<FIter> AllIterators;
+	for (auto& Pair : SectionToKeyTimeCache)
 	{
-		if (FMath::IsNearlyEqual(KeyTime, NewTime))
+		AllIterators.Add(Pair.Value.GetKeysInRange(CachedViewRange));
+	}
+	for (auto& Uncached: UncachableKeyTimes)
+	{
+		AllIterators.Add(Uncached.GetKeysInRange(CachedViewRange));
+	}
+
+	FTimeToPixel TimeToPixelConverter(AllottedGeometry, CachedViewRange);
+
+	// While any iterator is still valid, find and add the earliest time
+	while (AllIterators.ContainsByPredicate([](FIter& It){ return It; }))
+	{
+		float EarliestTime = TNumericLimits<float>::Max();
+		for (FIter& It : AllIterators)
 		{
-			return;
+			if (It && It->Time < EarliestTime)
+			{
+				EarliestTime = It->Time;
+			}
 		}
-	}
 
-	OutKeyTimes.Add(NewTime);
+		// Add the position as a pixel position
+		const float KeyPosition = TimeToPixelConverter.TimeToPixel(EarliestTime);
+		KeyDrawPositions.Add(KeyPosition);
+
+		// Increment any other iterators that are close enough to the time we just added
+		for (FIter& It : AllIterators)
+		{
+			while (It && FMath::IsNearlyEqual(KeyPosition, TimeToPixelConverter.TimeToPixel(It->Time), DuplicateThresholdPx))
+			{
+				++It;
+			}
+		}
+	}
 }
 
 
@@ -632,26 +737,13 @@ bool FSequencerDisplayNode::IsHovered() const
 
 
 TSharedRef<FGroupedKeyArea> FSequencerDisplayNode::GetKeyGrouping(UMovieSceneSection* InSection)
-{
-	TSharedRef<FGroupedKeyArea>* KeyGroup = KeyGroupings.FindByPredicate([=](const TSharedRef<FGroupedKeyArea>& InKeyArea) { return InKeyArea->GetOwningSection() == InSection; });
-	if (KeyGroup)
-	{
-		return *KeyGroup;
-	}
-
-	KeyGroupings.Emplace(MakeShared<FGroupedKeyArea>(*this, InSection));
-	return KeyGroupings.Last();
-}
-
-
-TSharedRef<FGroupedKeyArea> FSequencerDisplayNode::UpdateKeyGrouping(UMovieSceneSection* InSection)
 {
 	TSharedRef<FGroupedKeyArea>* KeyGroup = KeyGroupings.FindByPredicate([=](const TSharedRef<FGroupedKeyArea>& InKeyArea) { return InKeyArea->GetOwningSection() == InSection; });
 	if (KeyGroup)
 	{
 		if (KeyGroupRegenerationLock.GetValue() == 0)
 		{
-			**KeyGroup = FGroupedKeyArea(*this, InSection);
+			(*KeyGroup)->Update();
 		}
 		return *KeyGroup;
 	}
diff --git a/Engine/Source/Editor/Sequencer/Private/DisplayNodes/SequencerDisplayNode.h b/Engine/Source/Editor/Sequencer/Private/DisplayNodes/SequencerDisplayNode.h
--- a/Engine/Source/Editor/Sequencer/Private/DisplayNodes/SequencerDisplayNode.h
+++ b/Engine/Source/Editor/Sequencer/Private/DisplayNodes/SequencerDisplayNode.h
@@ -391,13 +391,10 @@ public:
 	}
 
 	/** Get the key grouping for the specified section index, ensuring it is fully up to date */
-	TSharedRef<FGroupedKeyArea> UpdateKeyGrouping(UMovieSceneSection* InSection);
-
-	/** Get the key grouping for the specified section index */
 	TSharedRef<FGroupedKeyArea> GetKeyGrouping(UMovieSceneSection* InSection);
 
 	/** Get key groupings array */
-	const TArray<TSharedRef<FGroupedKeyArea>> GetKeyGroupings() const { return KeyGroupings; }
+	const TArray<TSharedRef<FGroupedKeyArea>>& GetKeyGroupings() const { return KeyGroupings; }
 
 	DECLARE_EVENT(FSequencerDisplayNode, FRequestRenameEvent);
 	FRequestRenameEvent& OnRenameRequested() { return RenameRequestedEvent; }
diff --git a/Engine/Source/Editor/Sequencer/Private/DisplayNodes/SequencerKeyTimeCache.cpp b/Engine/Source/Editor/Sequencer/Private/DisplayNodes/SequencerKeyTimeCache.cpp
new file mode 100644
--- /dev/null
+++ b/Engine/Source/Editor/Sequencer/Private/DisplayNodes/SequencerKeyTimeCache.cpp
@@ -0,0 +1,46 @@
+// Copyright 1998-2017 Epic Games, Inc. All Rights Reserved.
+
+#include "SequencerKeyTimeCache.h"
+#include "MovieSceneSection.h"
+#include "Algo/Sort.h"
+#include "Algo/BinarySearch.h"
+
+void FSequencerCachedKeys::Update(TSharedRef<IKeyArea> KeyArea)
+{
+	if (!CachedKeys.IsSet())
+	{
+		CachedKeys.Emplace();
+	}
+
+	UMovieSceneSection* Section = KeyArea->GetOwningSection();
+	if (!Section || !CachedSignature.IsValid() || Section->GetSignature() != CachedSignature)
+	{
+		CachedSignature = Section ? Section->GetSignature() : FGuid();
+
+		CachedKeys->Reset();
+
+		// Generate and cache
+		for (FKeyHandle KeyHandle : KeyArea->GetUnsortedKeyHandles())
+		{
+			CachedKeys->Add(FSequencerCachedKey{ KeyHandle, KeyArea->GetKeyTime(KeyHandle) });
+		}
+
+		Algo::SortBy(CachedKeys.GetValue(), &FSequencerCachedKey::Time);
+	}
+}
+
+TArrayView<const FSequencerCachedKey> FSequencerCachedKeys::GetKeysInRange(TRange<float> ViewRange) const
+{
+	// Binary search the first time that's >= the lower bound
+	int32 FirstVisibleIndex = Algo::LowerBoundBy(CachedKeys.GetValue(), ViewRange.GetLowerBoundValue(), &FSequencerCachedKey::Time, TLess<float>());
+	// Binary search the last time that's > the upper bound
+	int32 LastVisibleIndex = Algo::UpperBoundBy(CachedKeys.GetValue(), ViewRange.GetUpperBoundValue(), &FSequencerCachedKey::Time, TLess<float>());
+
+	int32 Num = LastVisibleIndex - FirstVisibleIndex;
+	if (CachedKeys->IsValidIndex(FirstVisibleIndex) && LastVisibleIndex <= CachedKeys->Num())
+	{
+		return MakeArrayView(&CachedKeys.GetValue()[FirstVisibleIndex], Num);
+	}
+
+	return TArrayView<const FSequencerCachedKey>();
+}
\ No newline at end of file
diff --git a/Engine/Source/Editor/Sequencer/Private/DisplayNodes/SequencerKeyTimeCache.h b/Engine/Source/Editor/Sequencer/Private/DisplayNodes/SequencerKeyTimeCache.h
new file mode 100644
--- /dev/null
+++ b/Engine/Source/Editor/Sequencer/Private/DisplayNodes/SequencerKeyTimeCache.h
@@ -0,0 +1,38 @@
+// Copyright 1998-2017 Epic Games, Inc. All Rights Reserved.
+
+#pragma once
+
+#include "Range.h"
+#include "Guid.h"
+#include "ArrayView.h"
+#include "Curves/KeyHandle.h"
+
+#include "SharedPointer.h"
+#include "Optional.h"
+#include "IKeyArea.h"
+
+/** Key information that has been cached to avoid expensive operations */
+struct FSequencerCachedKey
+{
+	/** The key handle */
+	FKeyHandle Handle;
+	/** The local time of the key */
+	float Time;
+};
+
+/** Simple structure that caches the sorted key times for a given key area */
+struct FSequencerCachedKeys
+{
+	/** Update this cache to store key times and handles from the specified key area */
+	void Update(TSharedRef<IKeyArea> KeyArea);
+
+	/** Get an view of the cached array for keys that fall within the specified range */
+	TArrayView<const FSequencerCachedKey> GetKeysInRange(TRange<float> Range) const;
+
+private:
+	/** Cached key information */
+	TOptional<TArray<FSequencerCachedKey>> CachedKeys;
+	/** The guid with which the above array was generated */
+	FGuid CachedSignature;
+};
+
