Mirroring CL 4077578 from Dev-Editor, unblocking testers with VS 2017 packaging from UFE

Fixed compile during Client configuration of WmfMedia module

diff --git a/Engine/Plugins/Media/WmfMedia/Source/WmfMedia/Private/Wmf/WmfMediaByteStream.cpp b/Engine/Plugins/Media/WmfMedia/Source/WmfMedia/Private/Wmf/WmfMediaByteStream.cpp
--- a/Engine/Plugins/Media/WmfMedia/Source/WmfMedia/Private/Wmf/WmfMediaByteStream.cpp
+++ b/Engine/Plugins/Media/WmfMedia/Source/WmfMedia/Private/Wmf/WmfMediaByteStream.cpp
@@ -144,7 +144,7 @@ STDMETHODIMP FWmfMediaByteStream::BeginRead(BYTE* pb, ULONG cb, IMFAsyncCallback
 		return E_INVALIDARG;
 	}
 
-	TComPtr<FWmfMediaReadState> ReadState = new(std::nothrow) FWmfMediaReadState(pb, cb);
+	TComPtr<FWmfMediaReadState> ReadState = new FWmfMediaReadState(pb, cb);
 
 	if (ReadState == NULL)
 	{
