Fixed changing media tracks causing video to restart from beginning on Windows.

diff --git a/Engine/Plugins/Media/WmfMedia/Source/WmfMedia/Private/Wmf/WmfMediaSession.cpp b/Engine/Plugins/Media/WmfMedia/Source/WmfMedia/Private/Wmf/WmfMediaSession.cpp
--- a/Engine/Plugins/Media/WmfMedia/Source/WmfMedia/Private/Wmf/WmfMediaSession.cpp
+++ b/Engine/Plugins/Media/WmfMedia/Source/WmfMedia/Private/Wmf/WmfMediaSession.cpp
@@ -930,6 +930,10 @@ bool FWmfMediaSession::CommitTopology(IMFTopology* Topology)
 			UE_LOG(LogWmfMedia, Verbose, TEXT("Session %p: Requesting restart after pending stop"), this);
 
 			RequestedTime = LastTime;
+
+			// Zero LastTime so it matches what WMF is actually doing (rewinding to the beginning)
+			// The PresentationClock has now stopped and FWmfMediaSession::GetTime() will just return LastTime
+			LastTime = FTimespan::Zero();
 		}
 
 		UE_LOG(LogWmfMedia, Verbose, TEXT("Session %p: Requesting topology change after pending stop"), this);
