Update to latest RHI API to allow deprecation.

diff --git a/Engine/Plugins/Media/WmfMedia/Source/WmfMedia/Private/Wmf/WmfMediaHardwareVideoDecodingRendering.cpp b/Engine/Plugins/Media/WmfMedia/Source/WmfMedia/Private/Wmf/WmfMediaHardwareVideoDecodingRendering.cpp
--- a/Engine/Plugins/Media/WmfMedia/Source/WmfMedia/Private/Wmf/WmfMediaHardwareVideoDecodingRendering.cpp
+++ b/Engine/Plugins/Media/WmfMedia/Source/WmfMedia/Private/Wmf/WmfMediaHardwareVideoDecodingRendering.cpp
@@ -53,8 +53,8 @@ bool FWmfMediaHardwareVideoDecodingParameters::ConvertTextureFormat_RenderThread
 	{
 		FRHICommandListImmediate& RHICmdList = FRHICommandListExecutor::GetImmediateCommandList();
 
-		// Set render target.
-		SetRenderTarget(RHICmdList, InDstTexture, FTextureRHIRef(), ESimpleRenderTargetMode::EClearColorAndDepth, FExclusiveDepthStencil::DepthNop_StencilNop);
+		FRHIRenderPassInfo RPInfo(InDstTexture, ERenderTargetActions::DontLoad_Store);
+		RHICmdList.BeginRenderPass(RPInfo, TEXT("ConvertTextureFormat"));
 
 		// Update viewport.
 		RHICmdList.SetViewport(0, 0, 0.f, InSample->GetDim().X, InSample->GetDim().Y, 1.f);
@@ -134,6 +134,7 @@ bool FWmfMediaHardwareVideoDecodingParameters::ConvertTextureFormat_RenderThread
 		VertexShader->SetParameters(RHICmdList, VertexShader->GetVertexShader(), Y_SRV, UV_SRV, InSample->IsOutputSrgb());
 		PixelShader->SetParameters(RHICmdList, PixelShader->GetPixelShader(), Y_SRV, UV_SRV, InSample->IsOutputSrgb());
 		RHICmdList.DrawPrimitive(0, 2, 1);
+		RHICmdList.EndRenderPass();
 
 		D3D11DeviceContext->Release();
 	}
