Prevent unhandled exception for users who use Windows 10 "N" without having installed the Media Feature Pack.

#rnx

diff --git a/Engine/Plugins/Media/WmfMedia/Source/WmfMedia/Private/WmfMediaModule.cpp b/Engine/Plugins/Media/WmfMedia/Source/WmfMedia/Private/WmfMediaModule.cpp
--- a/Engine/Plugins/Media/WmfMedia/Source/WmfMedia/Private/WmfMediaModule.cpp
+++ b/Engine/Plugins/Media/WmfMedia/Source/WmfMedia/Private/WmfMediaModule.cpp
@@ -69,6 +69,11 @@ public:
 
 	//~ IWmfMediaModule interface
 
+	virtual bool IsInitialized() const override
+	{
+		return Initialized;
+	}
+
 	virtual TSharedPtr<IMediaPlayer, ESPMode::ThreadSafe> CreatePlayer(IMediaEventSink& EventSink) override
 	{
 #if WMFMEDIA_SUPPORTED_PLATFORM
diff --git a/Engine/Plugins/Media/WmfMedia/Source/WmfMedia/Public/IWmfMediaModule.h b/Engine/Plugins/Media/WmfMedia/Source/WmfMedia/Public/IWmfMediaModule.h
--- a/Engine/Plugins/Media/WmfMedia/Source/WmfMedia/Public/IWmfMediaModule.h
+++ b/Engine/Plugins/Media/WmfMedia/Source/WmfMedia/Public/IWmfMediaModule.h
@@ -17,6 +17,12 @@ class IWmfMediaModule
 {
 public:
 
+	/**
+	 * Is the Wmf media module initialized?
+	 @return True if the module is initialized.
+	 */
+	virtual bool IsInitialized() const = 0;
+
 	/**
 	 * Creates a Windows Media Foundation based media player.
 	 *
