Integrated CL# 5415512 & 5446812

Fix failed to create SRV on NV12 sub-textures issue
Reapply gamma correction shader parameter binding change

diff --git a/Engine/Plugins/Media/WmfMedia/Source/WmfMedia/Public/WmfMediaHardwareVideoDecodingShaders.h b/Engine/Plugins/Media/WmfMedia/Source/WmfMedia/Public/WmfMediaHardwareVideoDecodingShaders.h
--- a/Engine/Plugins/Media/WmfMedia/Source/WmfMedia/Public/WmfMediaHardwareVideoDecodingShaders.h
+++ b/Engine/Plugins/Media/WmfMedia/Source/WmfMedia/Public/WmfMediaHardwareVideoDecodingShaders.h
@@ -54,7 +54,7 @@ public:
 		SetSamplerParameter(RHICmdList, ShaderRHI, BilinearClampedSamplerUV, TStaticSamplerState<SF_Bilinear, AM_Clamp, AM_Clamp, AM_Clamp>::GetRHI());
 
 		SetShaderValue(RHICmdList, ShaderRHI, ColorTransform, MediaShaders::CombineColorTransformAndOffset(MediaShaders::YuvToSrgbDefault, MediaShaders::YUVOffset8bits));
-		SetShaderValue(RHICmdList, ShaderRHI, SrgbToLinear, InIsOutputSrgb);
+		SetShaderValue(RHICmdList, ShaderRHI, SrgbToLinear, InIsOutputSrgb ? 1 : 0); // Explicitly specify integer value, as using boolean falls over on XboxOne.
 	}
 
 
