Integrated CL# 5398218

Fixed some wrong #ifdefs that were wrongly skipping code in Editor builds.

diff --git a/Engine/Plugins/Media/WmfMedia/Source/WmfMedia/Private/Wmf/WmfMediaHardwareVideoDecodingTextureSample.cpp b/Engine/Plugins/Media/WmfMedia/Source/WmfMedia/Private/Wmf/WmfMediaHardwareVideoDecodingTextureSample.cpp
--- a/Engine/Plugins/Media/WmfMedia/Source/WmfMedia/Private/Wmf/WmfMediaHardwareVideoDecodingTextureSample.cpp
+++ b/Engine/Plugins/Media/WmfMedia/Source/WmfMedia/Private/Wmf/WmfMediaHardwareVideoDecodingTextureSample.cpp
@@ -46,7 +46,7 @@ ID3D11Texture2D* FWmfMediaHardwareVideoDecodingTextureSample::InitializeSourceTe
 	return SourceTexture;
 }
 
-#if !WITH_SERVER_CODE
+#if !UE_SERVER
 void FWmfMediaHardwareVideoDecodingTextureSample::ShutdownPoolable()
 {
 	FWmfMediaTextureSample::ShutdownPoolable();
diff --git a/Engine/Plugins/Media/WmfMedia/Source/WmfMedia/Public/WmfMediaHardwareVideoDecodingTextureSample.h b/Engine/Plugins/Media/WmfMedia/Source/WmfMedia/Public/WmfMediaHardwareVideoDecodingTextureSample.h
--- a/Engine/Plugins/Media/WmfMedia/Source/WmfMedia/Public/WmfMediaHardwareVideoDecodingTextureSample.h
+++ b/Engine/Plugins/Media/WmfMedia/Source/WmfMedia/Public/WmfMediaHardwareVideoDecodingTextureSample.h
@@ -105,7 +105,7 @@ public:
 	/**
 	 * Called the the sample is returned to the pool for cleanup purposes
 	 */
-#if !WITH_SERVER_CODE
+#if !UE_SERVER
 	virtual void ShutdownPoolable() override;
 #endif
 
