Move code that access d3d11 immediate context to an RHI command. With the d3d11 RHI thread enabled, it is no longer safe to access immediate context on RT

diff --git a/Engine/Plugins/Media/WmfMedia/Source/WmfMedia/Private/Wmf/WmfMediaHardwareVideoDecodingRendering.cpp b/Engine/Plugins/Media/WmfMedia/Source/WmfMedia/Private/Wmf/WmfMediaHardwareVideoDecodingRendering.cpp
--- a/Engine/Plugins/Media/WmfMedia/Source/WmfMedia/Private/Wmf/WmfMediaHardwareVideoDecodingRendering.cpp
+++ b/Engine/Plugins/Media/WmfMedia/Source/WmfMedia/Private/Wmf/WmfMediaHardwareVideoDecodingRendering.cpp
@@ -29,6 +29,77 @@
 
 #include "WmfMediaHardwareVideoDecodingShaders.h"
 
+struct FRHICommandCopyResource final : public FRHICommand<FRHICommandCopyResource>
+{
+	TComPtr<ID3D11Texture2D> SampleTexture;
+	FTexture2DRHIRef SampleDestinationTexture;
+
+	FRHICommandCopyResource(ID3D11Texture2D* InSampleTexture, FTexture2DRHIParamRef InSampleDestinationTexture)
+		: SampleTexture(InSampleTexture)
+		, SampleDestinationTexture(InSampleDestinationTexture)
+	{
+	}
+
+	void Execute(FRHICommandListBase& CmdList)
+	{
+		ID3D11Device* D3D11Device = static_cast<ID3D11Device*>(GDynamicRHI->RHIGetNativeDevice());
+		ID3D11DeviceContext* D3D11DeviceContext = nullptr;
+
+		D3D11Device->GetImmediateContext(&D3D11DeviceContext);
+		if (D3D11DeviceContext)
+		{
+			ID3D11Resource* DestinationTexture = reinterpret_cast<ID3D11Resource*>(SampleDestinationTexture->GetNativeResource());
+			if (DestinationTexture)
+			{
+				TComPtr<IDXGIResource> OtherResource(nullptr);
+				SampleTexture->QueryInterface(__uuidof(IDXGIResource), (void**)&OtherResource);
+
+				if (OtherResource)
+				{
+					HANDLE SharedHandle = nullptr;
+					if (OtherResource->GetSharedHandle(&SharedHandle) == S_OK)
+					{
+						if (SharedHandle != 0)
+						{
+							TComPtr<ID3D11Resource> SharedResource;
+							D3D11Device->OpenSharedResource(SharedHandle, __uuidof(ID3D11Texture2D), (void**)&SharedResource);
+
+							if (SharedResource)
+							{
+								TComPtr<IDXGIKeyedMutex> KeyedMutex;
+								SharedResource->QueryInterface(_uuidof(IDXGIKeyedMutex), (void**)&KeyedMutex);
+
+								if (KeyedMutex)
+								{
+									// Key is 1 : Texture as just been updated
+									// Key is 2 : Texture as already been updated.
+									// Do not wait to acquire key 1 since there is race no condition between writer and reader.
+									if (KeyedMutex->AcquireSync(1, 0) == S_OK)
+									{
+										// Copy from shared texture of FWmfMediaSink device to Rendering device
+										D3D11DeviceContext->CopyResource(DestinationTexture, SharedResource);
+										KeyedMutex->ReleaseSync(2);
+									}
+									else
+									{
+										// If key 1 cannot be acquired, another reader is already copying the resource
+										// and will release key with 2. 
+										// Wait to acquire key 2.
+										if (KeyedMutex->AcquireSync(2, INFINITE) == S_OK)
+										{
+											KeyedMutex->ReleaseSync(2);
+										}
+									}
+								}
+							}
+						}
+					}
+				}
+			}
+		}
+	}
+};
+
 bool FWmfMediaHardwareVideoDecodingParameters::ConvertTextureFormat_RenderThread(FWmfMediaHardwareVideoDecodingTextureSample* InSample, FTexture2DRHIRef InDstTexture)
 {
 	if (InSample == nullptr || !InDstTexture.IsValid())
@@ -75,53 +146,14 @@ bool FWmfMediaHardwareVideoDecodingParameters::ConvertTextureFormat_RenderThread
 
 		FTexture2DRHIRef SampleDestinationTexture = InSample->GetOrCreateDestinationTexture();
 
-		ID3D11Resource* DestinationTexture = reinterpret_cast<ID3D11Resource*>(SampleDestinationTexture->GetNativeResource());
-		if (DestinationTexture)
+		if (RHICmdList.Bypass())
 		{
-			TComPtr<IDXGIResource> OtherResource(nullptr);
-			SampleTexture->QueryInterface(__uuidof(IDXGIResource), (void**)&OtherResource);
-
-			if (OtherResource)
-			{
-				HANDLE SharedHandle = nullptr;
-				if (OtherResource->GetSharedHandle(&SharedHandle) == S_OK)
-				{
-					if (SharedHandle != 0)
-					{
-						TComPtr<ID3D11Resource> SharedResource;
-						D3D11Device->OpenSharedResource(SharedHandle, __uuidof(ID3D11Texture2D), (void**)&SharedResource);
-
-						if (SharedResource)
-						{
-							TComPtr<IDXGIKeyedMutex> KeyedMutex;
-							SharedResource->QueryInterface(_uuidof(IDXGIKeyedMutex), (void**)&KeyedMutex);
-
-							if (KeyedMutex)
-							{
-								// Key is 1 : Texture as just been updated
-								// Key is 2 : Texture as already been updated.
-								// Do not wait to acquire key 1 since there is race no condition between writer and reader.
-								if (KeyedMutex->AcquireSync(1, 0) == S_OK)
-								{
-									// Copy from shared texture of FWmfMediaSink device to Rendering device
-									D3D11DeviceContext->CopyResource(DestinationTexture, SharedResource);
-									KeyedMutex->ReleaseSync(2);
-								}
-								else
-								{
-									// If key 1 cannot be acquired, another reader is already copying the resource
-									// and will release key with 2. 
-									// Wait to acquire key 2.
-									if (KeyedMutex->AcquireSync(2, INFINITE) == S_OK)
-									{
-										KeyedMutex->ReleaseSync(2);
-									}
-								}
-							}
-						}
-					}
-				}
-			}
+			FRHICommandCopyResource Cmd(SampleTexture, SampleDestinationTexture);
+			Cmd.Execute(RHICmdList);
+		}
+		else
+		{
+			new (RHICmdList.AllocCommand<FRHICommandCopyResource>()) FRHICommandCopyResource(SampleTexture, SampleDestinationTexture);
 		}
 
 		FShaderResourceViewRHIRef Y_SRV = RHICreateShaderResourceView(SampleDestinationTexture, 0, 1, PF_G8);
