Need to call Release on the immediate context because GetImmediateContext bumps its ref count

diff --git a/Engine/Plugins/Media/WmfMedia/Source/WmfMedia/Private/Wmf/WmfMediaHardwareVideoDecodingRendering.cpp b/Engine/Plugins/Media/WmfMedia/Source/WmfMedia/Private/Wmf/WmfMediaHardwareVideoDecodingRendering.cpp
--- a/Engine/Plugins/Media/WmfMedia/Source/WmfMedia/Private/Wmf/WmfMediaHardwareVideoDecodingRendering.cpp
+++ b/Engine/Plugins/Media/WmfMedia/Source/WmfMedia/Private/Wmf/WmfMediaHardwareVideoDecodingRendering.cpp
@@ -96,6 +96,7 @@ struct FRHICommandCopyResource final : public FRHICommand<FRHICommandCopyResourc
 					}
 				}
 			}
+			D3D11DeviceContext->Release();
 		}
 	}
 };
