Fix race condition when using sample pool in hardware accelerated video decoding.

diff --git a/Engine/Plugins/Media/WmfMedia/Source/WmfMedia/Private/Player/WmfMediaTracks.cpp b/Engine/Plugins/Media/WmfMedia/Source/WmfMedia/Private/Player/WmfMediaTracks.cpp
--- a/Engine/Plugins/Media/WmfMedia/Source/WmfMedia/Private/Player/WmfMediaTracks.cpp
+++ b/Engine/Plugins/Media/WmfMedia/Source/WmfMedia/Private/Player/WmfMediaTracks.cpp
@@ -62,10 +62,6 @@ FWmfMediaTracks::~FWmfMediaTracks()
 
 	delete VideoSamplePool;
 	VideoSamplePool = nullptr;
-
-	delete VideoHardwareVideoDecodingSamplePool;
-	VideoHardwareVideoDecodingSamplePool = nullptr;
-
 }
 
 
@@ -1077,12 +1073,7 @@ bool FWmfMediaTracks::AddTrackToTopology(const FTrack& Track, IMFTopology& Topol
 		MajorType == MFMediaType_Video &&
 		FWmfMediaStreamSink::Create(MFMediaType_Video, MediaStreamSink))
 	{
-		if (VideoHardwareVideoDecodingSamplePool)
-		{
-			delete VideoHardwareVideoDecodingSamplePool;
-		}
-
-		VideoHardwareVideoDecodingSamplePool = new FWmfMediaHardwareVideoDecodingTextureSamplePool();
+		VideoHardwareVideoDecodingSamplePool = MakeShared<FWmfMediaHardwareVideoDecodingTextureSamplePool>();
 
 		MediaStreamSink->SetMediaSamplePoolAndQueue(VideoHardwareVideoDecodingSamplePool, &VideoSampleQueue);
 
diff --git a/Engine/Plugins/Media/WmfMedia/Source/WmfMedia/Private/Player/WmfMediaTracks.h b/Engine/Plugins/Media/WmfMedia/Source/WmfMedia/Private/Player/WmfMediaTracks.h
--- a/Engine/Plugins/Media/WmfMedia/Source/WmfMedia/Private/Player/WmfMediaTracks.h
+++ b/Engine/Plugins/Media/WmfMedia/Source/WmfMedia/Private/Player/WmfMediaTracks.h
@@ -334,7 +334,7 @@ private:
 
 	/** Video sample object pool. */
 	FWmfMediaTextureSamplePool* VideoSamplePool;
-	FWmfMediaHardwareVideoDecodingTextureSamplePool* VideoHardwareVideoDecodingSamplePool;
+	TSharedPtr<FWmfMediaHardwareVideoDecodingTextureSamplePool> VideoHardwareVideoDecodingSamplePool;
 
 	/** Video sample queue. */
 	TMediaSampleQueue<IMediaTextureSample> VideoSampleQueue;
diff --git a/Engine/Plugins/Media/WmfMedia/Source/WmfMedia/Private/Wmf/WmfMediaStreamSink.cpp b/Engine/Plugins/Media/WmfMedia/Source/WmfMedia/Private/Wmf/WmfMediaStreamSink.cpp
--- a/Engine/Plugins/Media/WmfMedia/Source/WmfMedia/Private/Wmf/WmfMediaStreamSink.cpp
+++ b/Engine/Plugins/Media/WmfMedia/Source/WmfMedia/Private/Wmf/WmfMediaStreamSink.cpp
@@ -931,7 +931,7 @@ void FWmfMediaStreamSink::SetClockRate(float InClockRate)
 
 
 void FWmfMediaStreamSink::SetMediaSamplePoolAndQueue(
-	FWmfMediaHardwareVideoDecodingTextureSamplePool* InVideoSamplePool,
+	TSharedPtr<FWmfMediaHardwareVideoDecodingTextureSamplePool>& InVideoSamplePool,
 	TMediaSampleQueue<IMediaTextureSample>* InVideoSampleQueue)
 {
 	FScopeLock Lock(&CriticalSection);
diff --git a/Engine/Plugins/Media/WmfMedia/Source/WmfMedia/Private/Wmf/WmfMediaStreamSink.h b/Engine/Plugins/Media/WmfMedia/Source/WmfMedia/Private/Wmf/WmfMediaStreamSink.h
--- a/Engine/Plugins/Media/WmfMedia/Source/WmfMedia/Private/Wmf/WmfMediaStreamSink.h
+++ b/Engine/Plugins/Media/WmfMedia/Source/WmfMedia/Private/Wmf/WmfMediaStreamSink.h
@@ -199,7 +199,7 @@ public:
 	 * @param InVideoSamplePool sample pool to get sample from
 	 * @param InVideoSampleQueue Sample queue to fill
 	 */
-	void SetMediaSamplePoolAndQueue(FWmfMediaHardwareVideoDecodingTextureSamplePool* InVideoSamplePool, TMediaSampleQueue<IMediaTextureSample>* InVideoSampleQueue);
+	void SetMediaSamplePoolAndQueue(TSharedPtr<FWmfMediaHardwareVideoDecodingTextureSamplePool>& InVideoSamplePool, TMediaSampleQueue<IMediaTextureSample>* InVideoSampleQueue);
 
 private:
 
@@ -268,7 +268,7 @@ private:
 	HANDLE WaitTimer;
 
 	/** Video sample pool from which to get next free sample */
-	FWmfMediaHardwareVideoDecodingTextureSamplePool* VideoSamplePool;
+	TSharedPtr<FWmfMediaHardwareVideoDecodingTextureSamplePool> VideoSamplePool;
 
 	/** Video sample queue which is filled with current sample*/
 	TMediaSampleQueue<IMediaTextureSample>* VideoSampleQueue;
