Fix freeze on WmfMediaSession shutdown when an error previously occured.

diff --git a/Engine/Plugins/Media/WmfMedia/Source/WmfMedia/Private/Wmf/WmfMediaSession.cpp b/Engine/Plugins/Media/WmfMedia/Source/WmfMedia/Private/Wmf/WmfMediaSession.cpp
--- a/Engine/Plugins/Media/WmfMedia/Source/WmfMedia/Private/Wmf/WmfMediaSession.cpp
+++ b/Engine/Plugins/Media/WmfMedia/Source/WmfMedia/Private/Wmf/WmfMediaSession.cpp
@@ -242,6 +242,8 @@ void FWmfMediaSession::Shutdown()
 		DiscardPendingChanges();
 	}
 
+	// When an error occurs we close the MediaSession in HandleError(), no need to Close it again.
+	if (SessionState != EMediaState::Error)
 	{
 		// We cannot have a lock here since we are waiting for an event which
 		// will call FWmfMediaSession::Invoke and would cause a deadlock
