Make sure we reset the session when switching stream when hardware accelerated video decoding is enabled.

diff --git a/Engine/Plugins/Media/WmfMedia/Source/WmfMedia/Private/Player/WmfMediaPlayer.cpp b/Engine/Plugins/Media/WmfMedia/Source/WmfMedia/Private/Player/WmfMediaPlayer.cpp
--- a/Engine/Plugins/Media/WmfMedia/Source/WmfMedia/Private/Player/WmfMediaPlayer.cpp
+++ b/Engine/Plugins/Media/WmfMedia/Source/WmfMedia/Private/Player/WmfMediaPlayer.cpp
@@ -175,8 +175,8 @@ void FWmfMediaPlayer::TickFetch(FTimespan /*DeltaTime*/, FTimespan /*Timecode*/)
 
 	if (TrackSelectionChanged)
 	{
-		// less than windows 10, seem to be a problem switching stream
-		if (!FWindowsPlatformMisc::VerifyWindowsVersion(10, 0) /* Anything < Windows 10.0 */)
+		// less than windows 10, seem to be a problem switching stream. The issue is also present when hardware acceleration is enabled.
+		if (!FWindowsPlatformMisc::VerifyWindowsVersion(10, 0) /* Anything < Windows 10.0 */ || GetDefault<UWmfMediaSettings>()->HardwareAcceleratedVideoDecoding)
 		{
 			const auto Settings = GetDefault<UWmfMediaSettings>();
 			check(Settings != nullptr);
