Integrate the missed 2 files from 6628601

#rnx

diff --git a/Engine/Plugins/Media/WmfMedia/Source/WmfMedia/Private/Wmf/WmfMediaHardwareVideoDecodingTextureSample.cpp b/Engine/Plugins/Media/WmfMedia/Source/WmfMedia/Private/Wmf/WmfMediaHardwareVideoDecodingTextureSample.cpp
--- a/Engine/Plugins/Media/WmfMedia/Source/WmfMedia/Private/Wmf/WmfMediaHardwareVideoDecodingTextureSample.cpp
+++ b/Engine/Plugins/Media/WmfMedia/Source/WmfMedia/Private/Wmf/WmfMediaHardwareVideoDecodingTextureSample.cpp
@@ -53,7 +53,10 @@ void FWmfMediaHardwareVideoDecodingTextureSample::ShutdownPoolable()
 
 	// Correctly release the keyed mutex when the sample is returned to the pool
 	TComPtr<IDXGIResource> OtherResource(nullptr);
-	SourceTexture->QueryInterface(__uuidof(IDXGIResource), (void**)&OtherResource);
+	if (SourceTexture)
+	{
+		SourceTexture->QueryInterface(__uuidof(IDXGIResource), (void**)&OtherResource);
+	}
 
 	if (OtherResource)
 	{
diff --git a/Engine/Plugins/Media/WmfMedia/Source/WmfMedia/Public/WmfMediaHardwareVideoDecodingTextureSample.h b/Engine/Plugins/Media/WmfMedia/Source/WmfMedia/Public/WmfMediaHardwareVideoDecodingTextureSample.h
--- a/Engine/Plugins/Media/WmfMedia/Source/WmfMedia/Public/WmfMediaHardwareVideoDecodingTextureSample.h
+++ b/Engine/Plugins/Media/WmfMedia/Source/WmfMedia/Public/WmfMediaHardwareVideoDecodingTextureSample.h
@@ -55,7 +55,8 @@ public:
 	 */
 	virtual IMediaTextureSampleConverter* GetMediaTextureSampleConverter() override
 	{
-		return this;
+		// Only use sample converter for Win8+
+		return FWindowsPlatformMisc::VerifyWindowsVersion(6, 2) ? this : nullptr;
 	}
 
 	/**
