Can now pause windows media using 'Native Audio Out'

diff --git a/Engine/Plugins/Media/WmfMedia/Source/WmfMedia/Private/Wmf/WmfMediaSession.cpp b/Engine/Plugins/Media/WmfMedia/Source/WmfMedia/Private/Wmf/WmfMediaSession.cpp
--- a/Engine/Plugins/Media/WmfMedia/Source/WmfMedia/Private/Wmf/WmfMediaSession.cpp
+++ b/Engine/Plugins/Media/WmfMedia/Source/WmfMedia/Private/Wmf/WmfMediaSession.cpp
@@ -1140,6 +1140,13 @@ void FWmfMediaSession::UpdateCharacteristics()
 		{
 			UnthinnedRates.Add(TRange<float>::Inclusive(MaxRate, MinRate));
 		}
+
+		// When native out is enabled, the slowest rate will be greater than 0.0f
+		bool SupportsUnthinnedPause = SUCCEEDED(RateSupport->IsRateSupported(FALSE, 0.0f, NULL));
+		if (SupportsUnthinnedPause && !ThinnedRates.Contains(0.0f) && !UnthinnedRates.Contains(0.0f))
+		{
+			UnthinnedRates.Add(TRange<float>::Inclusive(0.0f, 0.0f));
+		}
 	}
 }
 
