Restore backed out CL 4982351
Implemented WMF hardware accelerated video decoding support.

#rnx

diff --git a/Engine/Plugins/Media/WmfMedia/Source/WmfMedia/Private/Wmf/WmfMediaHardwareVideoDecodingRendering.cpp b/Engine/Plugins/Media/WmfMedia/Source/WmfMedia/Private/Wmf/WmfMediaHardwareVideoDecodingRendering.cpp
--- a/Engine/Plugins/Media/WmfMedia/Source/WmfMedia/Private/Wmf/WmfMediaHardwareVideoDecodingRendering.cpp
+++ b/Engine/Plugins/Media/WmfMedia/Source/WmfMedia/Private/Wmf/WmfMediaHardwareVideoDecodingRendering.cpp
@@ -26,103 +26,7 @@
 #include "d3d11.h"
 #include "Windows/HideWindowsPlatformTypes.h"
 
-class FWmfMediaHardwareVideoDecodingShader : public FGlobalShader
-{
-public:
-	static bool ShouldCompilePermutation(const FGlobalShaderPermutationParameters& Parameters)
-	{
-		return IsFeatureLevelSupported(Parameters.Platform, ERHIFeatureLevel::SM4);
-	}
-
-	static void ModifyCompilationEnvironment(const FGlobalShaderPermutationParameters& Parameters, FShaderCompilerEnvironment& OutEnvironment)
-	{
-		FGlobalShader::ModifyCompilationEnvironment(Parameters, OutEnvironment);
-	}
-
-	FWmfMediaHardwareVideoDecodingShader() {}
-
-	FWmfMediaHardwareVideoDecodingShader(const ShaderMetaType::CompiledShaderInitializerType& Initializer)
-		: FGlobalShader(Initializer)
-	{
-		TextureY.Bind(Initializer.ParameterMap, TEXT("TextureY"));
-		TextureUV.Bind(Initializer.ParameterMap, TEXT("TextureUV"));
-
-		PointClampedSamplerY.Bind(Initializer.ParameterMap, TEXT("PointClampedSamplerY"));
-		BilinearClampedSamplerUV.Bind(Initializer.ParameterMap, TEXT("BilinearClampedSamplerUV"));
-		
-		ColorTransform.Bind(Initializer.ParameterMap, TEXT("ColorTransform"));
-		SrgbToLinear.Bind(Initializer.ParameterMap, TEXT("SrgbToLinear"));
-	}
-
-	template<typename TShaderRHIParamRef>
-	void SetParameters(
-		FRHICommandListImmediate& RHICmdList,
-		const TShaderRHIParamRef ShaderRHI,
-		const FShaderResourceViewRHIRef& InTextureY,
-		const FShaderResourceViewRHIRef& InTextureUV,
-		const bool InIsOutputSrgb
-		)
-	{
-		SetSRVParameter(RHICmdList, ShaderRHI, TextureY, InTextureY);
-		SetSRVParameter(RHICmdList, ShaderRHI, TextureUV, InTextureUV);
-		
-		SetSamplerParameter(RHICmdList, ShaderRHI, PointClampedSamplerY, TStaticSamplerState<SF_Point, AM_Clamp, AM_Clamp, AM_Clamp>::GetRHI());
-		SetSamplerParameter(RHICmdList, ShaderRHI, BilinearClampedSamplerUV, TStaticSamplerState<SF_Bilinear, AM_Clamp, AM_Clamp, AM_Clamp>::GetRHI());
-
-		SetShaderValue(RHICmdList, ShaderRHI, ColorTransform, MediaShaders::CombineColorTransformAndOffset(MediaShaders::YuvToSrgbDefault, MediaShaders::YUVOffset8bits));
-		SetShaderValue(RHICmdList, ShaderRHI, SrgbToLinear, InIsOutputSrgb);
-	}
-
-
-	virtual bool Serialize(FArchive& Ar) override
-	{
-		bool bShaderHasOutdatedParameters = FGlobalShader::Serialize(Ar);
-		Ar << TextureY << TextureUV << PointClampedSamplerY << BilinearClampedSamplerUV << ColorTransform << SrgbToLinear;
-		return bShaderHasOutdatedParameters;
-	}
-
-private:
-	FShaderResourceParameter TextureY;
-	FShaderResourceParameter TextureUV;
-	FShaderResourceParameter PointClampedSamplerY;
-	FShaderResourceParameter BilinearClampedSamplerUV;
-	FShaderParameter ColorTransform;
-	FShaderParameter SrgbToLinear;
-};
-
-class FHardwareVideoDecodingVS : public FWmfMediaHardwareVideoDecodingShader
-{
-	DECLARE_SHADER_TYPE(FHardwareVideoDecodingVS, Global);
-
-public:
-
-	/** Default constructor. */
-	FHardwareVideoDecodingVS() {}
-
-	/** Initialization constructor. */
-	FHardwareVideoDecodingVS(const ShaderMetaType::CompiledShaderInitializerType& Initializer)
-		: FWmfMediaHardwareVideoDecodingShader(Initializer)
-	{
-	}
-};
-
-class FHardwareVideoDecodingPS : public FWmfMediaHardwareVideoDecodingShader
-{
-	DECLARE_SHADER_TYPE(FHardwareVideoDecodingPS, Global);
-
-public:
-
-	/** Default constructor. */
-	FHardwareVideoDecodingPS() {}
-
-	/** Initialization constructor. */
-	FHardwareVideoDecodingPS(const ShaderMetaType::CompiledShaderInitializerType& Initializer)
-		: FWmfMediaHardwareVideoDecodingShader(Initializer)
-	{ }
-};
-
-IMPLEMENT_SHADER_TYPE(, FHardwareVideoDecodingVS, TEXT("/Plugin/WmfMedia/Private/MediaHardwareVideoDecoding.usf"), TEXT("MainVS"), SF_Vertex)
-IMPLEMENT_SHADER_TYPE(, FHardwareVideoDecodingPS, TEXT("/Plugin/WmfMedia/Private/MediaHardwareVideoDecoding.usf"), TEXT("NV12ConvertPS"), SF_Pixel)
+#include "WmfMediaHardwareVideoDecodingShaders.h"
 
 bool FWmfMediaHardwareVideoDecodingParameters::ConvertTextureFormat_RenderThread(FWmfMediaHardwareVideoDecodingTextureSample* InSample, FTexture2DRHIRef InDstTexture)
 {
diff --git a/Engine/Plugins/Media/WmfMedia/Source/WmfMedia/Private/Wmf/WmfMediaSink.cpp b/Engine/Plugins/Media/WmfMedia/Source/WmfMedia/Private/Wmf/WmfMediaSink.cpp
--- a/Engine/Plugins/Media/WmfMedia/Source/WmfMedia/Private/Wmf/WmfMediaSink.cpp
+++ b/Engine/Plugins/Media/WmfMedia/Source/WmfMedia/Private/Wmf/WmfMediaSink.cpp
@@ -580,7 +580,7 @@ bool FWmfMediaSink::CreateDXGIManagerAndDevice()
 			return false;
 		}
 
-		UE_LOG(LogWmfMedia, VeryVerbose, TEXT("D3D11 Device Created: %p"), D3D11Device.Get());
+		UE_LOG(LogWmfMedia, Log, TEXT("D3D11 Device Created: %p"), D3D11Device.Get());
 		return true;
 	}
 	else
diff --git a/Engine/Plugins/Media/WmfMedia/Source/WmfMedia/Private/WmfMediaModule.cpp b/Engine/Plugins/Media/WmfMedia/Source/WmfMedia/Private/WmfMediaModule.cpp
--- a/Engine/Plugins/Media/WmfMedia/Source/WmfMedia/Private/WmfMediaModule.cpp
+++ b/Engine/Plugins/Media/WmfMedia/Source/WmfMedia/Private/WmfMediaModule.cpp
@@ -4,6 +4,9 @@
 
 #include "IMediaCaptureSupport.h"
 #include "Modules/ModuleManager.h"
+#include "Interfaces/IPluginManager.h"
+#include "Misc/Paths.h"
+#include "ShaderCore.h"
 
 #include "IWmfMediaModule.h"
 
diff --git a/Engine/Plugins/Media/WmfMedia/Source/WmfMedia/Public/WmfMediaHardwareVideoDecodingRendering.h b/Engine/Plugins/Media/WmfMedia/Source/WmfMedia/Public/WmfMediaHardwareVideoDecodingRendering.h
--- a/Engine/Plugins/Media/WmfMedia/Source/WmfMedia/Public/WmfMediaHardwareVideoDecodingRendering.h
+++ b/Engine/Plugins/Media/WmfMedia/Source/WmfMedia/Public/WmfMediaHardwareVideoDecodingRendering.h
@@ -9,5 +9,5 @@ class FWmfMediaHardwareVideoDecodingTextureSample;
 class FWmfMediaHardwareVideoDecodingParameters
 {
 public:
-	static bool ConvertTextureFormat_RenderThread(FWmfMediaHardwareVideoDecodingTextureSample* InSample, FTexture2DRHIRef InDstTexture);
+	WMFMEDIA_API static bool ConvertTextureFormat_RenderThread(FWmfMediaHardwareVideoDecodingTextureSample* InSample, FTexture2DRHIRef InDstTexture);
 };
diff --git a/Engine/Plugins/Media/WmfMedia/Source/WmfMedia/Public/WmfMediaHardwareVideoDecodingShaders.h b/Engine/Plugins/Media/WmfMedia/Source/WmfMedia/Public/WmfMediaHardwareVideoDecodingShaders.h
new file mode 100644
--- /dev/null
+++ b/Engine/Plugins/Media/WmfMedia/Source/WmfMedia/Public/WmfMediaHardwareVideoDecodingShaders.h
@@ -0,0 +1,109 @@
+// Copyright 1998-2019 Epic Games, Inc. All Rights Reserved.
+
+#include "GlobalShader.h"
+#include "RenderUtils.h"
+#include "RHIStaticStates.h"
+#include "PipelineStateCache.h"
+#include "ShaderParameterUtils.h"
+#include "MediaShaders.h"
+
+/**
+ * Shaders which allow the conversion of NV12 texture data into RGB textures.
+ */
+class FWmfMediaHardwareVideoDecodingShader : public FGlobalShader
+{
+public:
+	static bool ShouldCompilePermutation(const FGlobalShaderPermutationParameters& Parameters)
+	{
+		return IsFeatureLevelSupported(Parameters.Platform, ERHIFeatureLevel::SM4);
+	}
+
+	static void ModifyCompilationEnvironment(const FGlobalShaderPermutationParameters& Parameters, FShaderCompilerEnvironment& OutEnvironment)
+	{
+		FGlobalShader::ModifyCompilationEnvironment(Parameters, OutEnvironment);
+	}
+
+	FWmfMediaHardwareVideoDecodingShader() {}
+
+	FWmfMediaHardwareVideoDecodingShader(const ShaderMetaType::CompiledShaderInitializerType& Initializer)
+		: FGlobalShader(Initializer)
+	{
+		TextureY.Bind(Initializer.ParameterMap, TEXT("TextureY"));
+		TextureUV.Bind(Initializer.ParameterMap, TEXT("TextureUV"));
+
+		PointClampedSamplerY.Bind(Initializer.ParameterMap, TEXT("PointClampedSamplerY"));
+		BilinearClampedSamplerUV.Bind(Initializer.ParameterMap, TEXT("BilinearClampedSamplerUV"));
+
+		ColorTransform.Bind(Initializer.ParameterMap, TEXT("ColorTransform"));
+		SrgbToLinear.Bind(Initializer.ParameterMap, TEXT("SrgbToLinear"));
+	}
+
+	template<typename TShaderRHIParamRef>
+	void SetParameters(
+		FRHICommandListImmediate& RHICmdList,
+		const TShaderRHIParamRef ShaderRHI,
+		const FShaderResourceViewRHIRef& InTextureY,
+		const FShaderResourceViewRHIRef& InTextureUV,
+		const bool InIsOutputSrgb
+	)
+	{
+		SetSRVParameter(RHICmdList, ShaderRHI, TextureY, InTextureY);
+		SetSRVParameter(RHICmdList, ShaderRHI, TextureUV, InTextureUV);
+
+		SetSamplerParameter(RHICmdList, ShaderRHI, PointClampedSamplerY, TStaticSamplerState<SF_Point, AM_Clamp, AM_Clamp, AM_Clamp>::GetRHI());
+		SetSamplerParameter(RHICmdList, ShaderRHI, BilinearClampedSamplerUV, TStaticSamplerState<SF_Bilinear, AM_Clamp, AM_Clamp, AM_Clamp>::GetRHI());
+
+		SetShaderValue(RHICmdList, ShaderRHI, ColorTransform, MediaShaders::CombineColorTransformAndOffset(MediaShaders::YuvToSrgbDefault, MediaShaders::YUVOffset8bits));
+		SetShaderValue(RHICmdList, ShaderRHI, SrgbToLinear, InIsOutputSrgb);
+	}
+
+
+	virtual bool Serialize(FArchive& Ar) override
+	{
+		bool bShaderHasOutdatedParameters = FGlobalShader::Serialize(Ar);
+		Ar << TextureY << TextureUV << PointClampedSamplerY << BilinearClampedSamplerUV << ColorTransform << SrgbToLinear;
+		return bShaderHasOutdatedParameters;
+	}
+
+private:
+	FShaderResourceParameter TextureY;
+	FShaderResourceParameter TextureUV;
+	FShaderResourceParameter PointClampedSamplerY;
+	FShaderResourceParameter BilinearClampedSamplerUV;
+	FShaderParameter ColorTransform;
+	FShaderParameter SrgbToLinear;
+};
+
+class FHardwareVideoDecodingVS : public FWmfMediaHardwareVideoDecodingShader
+{
+	DECLARE_SHADER_TYPE(FHardwareVideoDecodingVS, Global);
+
+public:
+
+	/** Default constructor. */
+	FHardwareVideoDecodingVS() {}
+
+	/** Initialization constructor. */
+	FHardwareVideoDecodingVS(const ShaderMetaType::CompiledShaderInitializerType& Initializer)
+		: FWmfMediaHardwareVideoDecodingShader(Initializer)
+	{
+	}
+};
+
+class FHardwareVideoDecodingPS : public FWmfMediaHardwareVideoDecodingShader
+{
+	DECLARE_SHADER_TYPE(FHardwareVideoDecodingPS, Global);
+
+public:
+
+	/** Default constructor. */
+	FHardwareVideoDecodingPS() {}
+
+	/** Initialization constructor. */
+	FHardwareVideoDecodingPS(const ShaderMetaType::CompiledShaderInitializerType& Initializer)
+		: FWmfMediaHardwareVideoDecodingShader(Initializer)
+	{ }
+};
+
+IMPLEMENT_SHADER_TYPE(, FHardwareVideoDecodingVS, TEXT("/Plugin/WmfMedia/Private/MediaHardwareVideoDecoding.usf"), TEXT("MainVS"), SF_Vertex)
+IMPLEMENT_SHADER_TYPE(, FHardwareVideoDecodingPS, TEXT("/Plugin/WmfMedia/Private/MediaHardwareVideoDecoding.usf"), TEXT("NV12ConvertPS"), SF_Pixel)
diff --git a/Engine/Plugins/Media/WmfMedia/Source/WmfMedia/WmfMedia.Build.cs b/Engine/Plugins/Media/WmfMedia/Source/WmfMedia/WmfMedia.Build.cs
--- a/Engine/Plugins/Media/WmfMedia/Source/WmfMedia/WmfMedia.Build.cs
+++ b/Engine/Plugins/Media/WmfMedia/Source/WmfMedia/WmfMedia.Build.cs
@@ -27,6 +27,11 @@ namespace UnrealBuildTool.Rules
                     "WmfMediaFactory",
                 });
 
+            if (Target.Platform != UnrealTargetPlatform.XboxOne)
+            {
+                PrivateDependencyModuleNames.Add("D3D11RHI");
+            }
+
 			PrivateIncludePathModuleNames.AddRange(
 				new string[] {
 					"Media",
diff --git a/Engine/Source/Runtime/RHI/Public/RHIDefinitions.h b/Engine/Source/Runtime/RHI/Public/RHIDefinitions.h
--- a/Engine/Source/Runtime/RHI/Public/RHIDefinitions.h
+++ b/Engine/Source/Runtime/RHI/Public/RHIDefinitions.h
@@ -714,7 +714,7 @@ enum ETextureCreateFlags
 	TexCreate_CPUWritable			= 1<<5,
 	// Texture will be created with an un-tiled format
 	TexCreate_NoTiling				= 1<<6,
-	// Texture will be used for video decode on Switch
+	// Texture will be used for video decode
 	TexCreate_VideoDecode			= 1<<7,
 	// Texture that may be updated every frame
 	TexCreate_Dynamic				= 1<<8,
