Restore backed out CL 4982351
Implemented WMF hardware accelerated video decoding support.

#rnx

diff --git a/Engine/Plugins/Media/WmfMedia/Source/WmfMedia/Public/WmfMediaHardwareVideoDecodingRendering.cpp b/Engine/Plugins/Media/WmfMedia/Source/WmfMedia/Public/WmfMediaHardwareVideoDecodingRendering.cpp
new file mode 100644
--- /dev/null
+++ b/Engine/Plugins/Media/WmfMedia/Source/WmfMedia/Public/WmfMediaHardwareVideoDecodingRendering.cpp
@@ -0,0 +1,140 @@
+// Copyright 1998-2018 Epic Games, Inc. All Rights Reserved.
+
+#if PLATFORM_WINDOWS
+
+#include "WmfMediaHardwareVideoDecodingRendering.h"
+#include "Engine/Texture2D.h"
+#include "Engine/TextureRenderTarget2D.h"
+#include "Engine/World.h"
+#include "Logging/LogMacros.h"
+#include "SceneUtils.h"
+#include "SceneInterface.h"
+
+#include "D3D11RHIPrivate.h"
+#include "DynamicRHI.h"
+
+#include "WmfMediaHardwareVideoDecodingTextureSample.h"
+#include "WmfMediaPrivate.h"
+
+
+#include "dxgi.h"
+
+#include "Windows/AllowWindowsPlatformTypes.h"
+#include "d3d11.h"
+#include "Windows/HideWindowsPlatformTypes.h"
+
+#include "Stats/Stats2.h"
+
+#include "WmfMediaHardwareVideoDecodingShaders.h"
+
+DECLARE_STATS_GROUP(TEXT("WmfMedia"), STATGROUP_WmfMedia, STATCAT_Advanced);
+
+DECLARE_CYCLE_STAT(TEXT("FWmfMediaHardwareVideoDecodingParameters::ConvertTextureFormat_RenderThread"), STAT_WmfMedia_FWmfMediaHardwareVideoDecodingParameters_ConvertTextureFormat_RenderThread, STATGROUP_WmfMedia);
+
+
+bool FWmfMediaHardwareVideoDecodingParameters::ConvertTextureFormat_RenderThread(FWmfMediaHardwareVideoDecodingTextureSample* InSample, FTexture2DRHIRef InDstTexture)
+{
+	if (InSample == nullptr || !InDstTexture.IsValid())
+	{
+		return false;
+	}
+
+	SCOPE_CYCLE_COUNTER(STAT_WmfMedia_FWmfMediaHardwareVideoDecodingParameters_ConvertTextureFormat_RenderThread);
+
+	check(IsInRenderingThread());
+	check(InSample);
+
+	TComPtr<ID3D11Texture2D> SampleTexture = InSample->GetSourceTexture();
+
+	ID3D11Device* D3D11Device = static_cast<ID3D11Device*>(GDynamicRHI->RHIGetNativeDevice());
+	ID3D11DeviceContext* D3D11DeviceContext = nullptr;
+			
+	// Must access rendering device context to copy shared resource.
+	D3D11Device->GetImmediateContext(&D3D11DeviceContext);
+	if (D3D11DeviceContext)
+	{
+		FRHICommandListImmediate& RHICmdList = FRHICommandListExecutor::GetImmediateCommandList();
+
+		// Set render target.
+		SetRenderTarget(RHICmdList, InDstTexture, FTextureRHIRef(), ESimpleRenderTargetMode::EClearColorAndDepth, FExclusiveDepthStencil::DepthNop_StencilNop);
+
+		// Update viewport.
+		RHICmdList.SetViewport(0, 0, 0.f, InSample->GetDim().X, InSample->GetDim().Y, 1.f);
+
+		// Get shaders.
+		TShaderMap<FGlobalShaderType>* GlobalShaderMap = GetGlobalShaderMap(GMaxRHIFeatureLevel);
+		TShaderMapRef< FHardwareVideoDecodingPS > PixelShader(GlobalShaderMap);
+		TShaderMapRef< FHardwareVideoDecodingVS > VertexShader(GlobalShaderMap);
+
+		// Set the graphic pipeline state.
+		FGraphicsPipelineStateInitializer GraphicsPSOInit;
+		RHICmdList.ApplyCachedRenderTargets(GraphicsPSOInit);
+		GraphicsPSOInit.DepthStencilState = TStaticDepthStencilState<false, CF_Always>::GetRHI();
+		GraphicsPSOInit.BlendState = TStaticBlendState<>::GetRHI();
+		GraphicsPSOInit.RasterizerState = TStaticRasterizerState<>::GetRHI();
+		GraphicsPSOInit.PrimitiveType = PT_TriangleList;
+		GraphicsPSOInit.BoundShaderState.VertexDeclarationRHI = GetVertexDeclarationFVector4();
+		GraphicsPSOInit.BoundShaderState.VertexShaderRHI = GETSAFERHISHADER_VERTEX(*VertexShader);
+		GraphicsPSOInit.BoundShaderState.PixelShaderRHI = GETSAFERHISHADER_PIXEL(*PixelShader);
+		SetGraphicsPipelineState(RHICmdList, GraphicsPSOInit);
+
+		FTexture2DRHIRef SampleDestinationTexture = InSample->GetDestinationTexture();
+		if (!SampleDestinationTexture.IsValid())
+		{
+			FRHIResourceCreateInfo CreateInfo;
+			const uint32 CreateFlags = TexCreate_Dynamic | TexCreate_DisableSRVCreation;
+			FTexture2DRHIRef Texture = RHICreateTexture2D(
+				InSample->GetDim().X,
+				InSample->GetDim().Y,
+				PF_NV12,
+				1,
+				1,
+				CreateFlags,
+				CreateInfo);
+
+			InSample->SetDestinationTexture(Texture);
+			SampleDestinationTexture = Texture;
+		}
+
+		ID3D11Resource* DestinationTexture = reinterpret_cast<ID3D11Resource*>(SampleDestinationTexture->GetNativeResource());
+		if (DestinationTexture)
+		{
+			TComPtr<IDXGIResource> OtherResource(nullptr);
+			SampleTexture->QueryInterface(__uuidof(IDXGIResource), (void**)&OtherResource);
+
+			if (OtherResource)
+			{
+				HANDLE sharedHandle = 0;
+				if (OtherResource->GetSharedHandle(&sharedHandle) == S_OK)
+				{
+					if (sharedHandle != 0)
+					{
+						ID3D11Resource* SharedResource = nullptr;
+						D3D11Device->OpenSharedResource(sharedHandle, __uuidof(ID3D11Texture2D), (void**)&SharedResource);
+
+						if (SharedResource)
+						{
+							// Copy from shared texture of FWmfMediaSink device to Rendering device
+							D3D11DeviceContext->CopyResource(DestinationTexture, SharedResource);
+							SharedResource->Release();
+						}
+					}
+				}
+			}
+		}
+
+		FShaderResourceViewRHIRef ChromaSRV = RHICreateShaderResourceView(SampleDestinationTexture, 0, 1, PF_G8);
+		FShaderResourceViewRHIRef LuminanceSRV = RHICreateShaderResourceView(SampleDestinationTexture, 0, 1, PF_R8G8);
+
+		// Update shader uniform parameters.
+		VertexShader->SetParameters(RHICmdList, VertexShader->GetVertexShader(), ChromaSRV, LuminanceSRV, InSample->IsOutputSrgb());
+		PixelShader->SetParameters(RHICmdList, PixelShader->GetPixelShader(), ChromaSRV, LuminanceSRV, InSample->IsOutputSrgb());
+		RHICmdList.DrawPrimitive(PT_TriangleList, 0, 2, 1);
+
+		D3D11DeviceContext->Release();
+	}
+
+	return true;
+}
+
+#endif
