Some fixes for build issues after integrate.

#rnx

diff --git a/Engine/Plugins/Media/WmfMedia/Source/WmfMedia/Private/Wmf/WmfMediaHardwareVideoDecodingTextureSample.cpp b/Engine/Plugins/Media/WmfMedia/Source/WmfMedia/Private/Wmf/WmfMediaHardwareVideoDecodingTextureSample.cpp
--- a/Engine/Plugins/Media/WmfMedia/Source/WmfMedia/Private/Wmf/WmfMediaHardwareVideoDecodingTextureSample.cpp
+++ b/Engine/Plugins/Media/WmfMedia/Source/WmfMedia/Private/Wmf/WmfMediaHardwareVideoDecodingTextureSample.cpp
@@ -2,7 +2,7 @@
 
 #include "WmfMediaHardwareVideoDecodingTextureSample.h"
 
-ID3D11Texture2D* FWmfMediaHardwareVideoDecodingTextureSample::InitializeSourceTexture(const TComPtr<ID3D11Device>& InD3D11Device, FTimespan InTime, FTimespan InDuration, const FIntPoint& InDim, uint8 InFormat, EMediaTextureSampleFormat InMediaTextureSampleFormat)
+ID3D11Texture2D* FWmfMediaHardwareVideoDecodingTextureSample::InitializeSourceTexture(const TRefCountPtr<ID3D11Device>& InD3D11Device, FTimespan InTime, FTimespan InDuration, const FIntPoint& InDim, uint8 InFormat, EMediaTextureSampleFormat InMediaTextureSampleFormat)
 {
 	Time = InTime;
 	Dim = InDim;
diff --git a/Engine/Plugins/Media/WmfMedia/Source/WmfMedia/Public/WmfMediaHardwareVideoDecodingTextureSample.h b/Engine/Plugins/Media/WmfMedia/Source/WmfMedia/Public/WmfMediaHardwareVideoDecodingTextureSample.h
--- a/Engine/Plugins/Media/WmfMedia/Source/WmfMedia/Public/WmfMediaHardwareVideoDecodingTextureSample.h
+++ b/Engine/Plugins/Media/WmfMedia/Source/WmfMedia/Public/WmfMediaHardwareVideoDecodingTextureSample.h
@@ -44,7 +44,7 @@ public:
 	 * @param InCreateFlags texture create flag
 	 * @return The texture resource object that will hold the sample data.
 	 */
-	ID3D11Texture2D* InitializeSourceTexture(const TComPtr<ID3D11Device>& InD3D11Device, FTimespan InTime, FTimespan InDuration, const FIntPoint& InDim, uint8 InFormat, EMediaTextureSampleFormat InMediaTextureSampleFormat);
+	ID3D11Texture2D* InitializeSourceTexture(const TRefCountPtr<ID3D11Device>& InD3D11Device, FTimespan InTime, FTimespan InDuration, const FIntPoint& InDim, uint8 InFormat, EMediaTextureSampleFormat InMediaTextureSampleFormat);
 
 	/**
 	 * Get media texture sample converter if sample implements it
@@ -105,7 +105,7 @@ private:
 	TComPtr<ID3D11Texture2D> SourceTexture;
 
 	/** D3D11 Device which create the texture, used to release the keyed mutex when the sampled is returned to the pool */
-	TComPtr<ID3D11Device> D3D11Device;
+	TRefCountPtr<ID3D11Device> D3D11Device;
 
 	/** Destination Texture resource (from Rendering device) */
 	FTexture2DRHIRef DestinationTexture;
