Some fixes for build issues after integrate.

#rnx

diff --git a/Engine/Plugins/Media/WmfMedia/Source/WmfMedia/Private/Wmf/WmfMediaHardwareVideoDecodingRendering.cpp b/Engine/Plugins/Media/WmfMedia/Source/WmfMedia/Private/Wmf/WmfMediaHardwareVideoDecodingRendering.cpp
--- a/Engine/Plugins/Media/WmfMedia/Source/WmfMedia/Private/Wmf/WmfMediaHardwareVideoDecodingRendering.cpp
+++ b/Engine/Plugins/Media/WmfMedia/Source/WmfMedia/Private/Wmf/WmfMediaHardwareVideoDecodingRendering.cpp
@@ -1,7 +1,11 @@
 // Copyright 1998-2019 Epic Games, Inc. All Rights Reserved.
 
+#include "WmfMediaPrivate.h"
+
 #include "WmfMediaHardwareVideoDecodingRendering.h"
 
+#if WMFMEDIA_SUPPORTED_PLATFORM
+
 #include "Engine/Texture2D.h"
 #include "Engine/TextureRenderTarget2D.h"
 #include "Engine/World.h"
@@ -152,3 +156,5 @@ bool FWmfMediaHardwareVideoDecodingParameters::ConvertTextureFormat_RenderThread
 
 	return true;
 }
+
+#endif
diff --git a/Engine/Plugins/Media/WmfMedia/Source/WmfMedia/Private/Wmf/WmfMediaHardwareVideoDecodingTextureSample.cpp b/Engine/Plugins/Media/WmfMedia/Source/WmfMedia/Private/Wmf/WmfMediaHardwareVideoDecodingTextureSample.cpp
--- a/Engine/Plugins/Media/WmfMedia/Source/WmfMedia/Private/Wmf/WmfMediaHardwareVideoDecodingTextureSample.cpp
+++ b/Engine/Plugins/Media/WmfMedia/Source/WmfMedia/Private/Wmf/WmfMediaHardwareVideoDecodingTextureSample.cpp
@@ -1,5 +1,9 @@
 // Copyright 1998-2019 Epic Games, Inc. All Rights Reserved.
 
+#include "WmfMediaPrivate.h"
+
+#if WMFMEDIA_SUPPORTED_PLATFORM
+
 #include "WmfMediaHardwareVideoDecodingTextureSample.h"
 
 ID3D11Texture2D* FWmfMediaHardwareVideoDecodingTextureSample::InitializeSourceTexture(const TRefCountPtr<ID3D11Device>& InD3D11Device, FTimespan InTime, FTimespan InDuration, const FIntPoint& InDim, uint8 InFormat, EMediaTextureSampleFormat InMediaTextureSampleFormat)
@@ -74,3 +78,4 @@ void FWmfMediaHardwareVideoDecodingTextureSample::ShutdownPoolable()
 	}
 }
 
+#endif
diff --git a/Engine/Plugins/Media/WmfMedia/Source/WmfMedia/WmfMedia.Build.cs b/Engine/Plugins/Media/WmfMedia/Source/WmfMedia/WmfMedia.Build.cs
--- a/Engine/Plugins/Media/WmfMedia/Source/WmfMedia/WmfMedia.Build.cs
+++ b/Engine/Plugins/Media/WmfMedia/Source/WmfMedia/WmfMedia.Build.cs
@@ -17,7 +17,6 @@ namespace UnrealBuildTool.Rules
 				new string[] {
 					"Core",
 					"CoreUObject",
-					"D3D11RHI",
                     "Engine",
                     "MediaUtils",
 					"Projects",
diff --git a/Engine/Plugins/Media/WmfMedia/WmfMedia.uplugin b/Engine/Plugins/Media/WmfMedia/WmfMedia.uplugin
--- a/Engine/Plugins/Media/WmfMedia/WmfMedia.uplugin
+++ b/Engine/Plugins/Media/WmfMedia/WmfMedia.uplugin
@@ -20,7 +20,7 @@
 			"Name" : "WmfMedia",
 			"Type" : "Runtime",
 			"LoadingPhase" : "PostConfigInit",
-			"WhitelistPlatforms" : [ "Win32", "Win64" ]
+			"WhitelistPlatforms" : [ "Win32", "Win64", "XboxOne" ]
 		},
 		{
 			"Name" : "WmfMediaEditor",
