Some fixes for build issues after integrate.
Some fixes for #includes

#rnx

diff --git a/Engine/Plugins/Media/WmfMedia/Source/WmfMedia/Public/WmfMediaHardwareVideoDecodingTextureSample.h b/Engine/Plugins/Media/WmfMedia/Source/WmfMedia/Public/WmfMediaHardwareVideoDecodingTextureSample.h
--- a/Engine/Plugins/Media/WmfMedia/Source/WmfMedia/Public/WmfMediaHardwareVideoDecodingTextureSample.h
+++ b/Engine/Plugins/Media/WmfMedia/Source/WmfMedia/Public/WmfMediaHardwareVideoDecodingTextureSample.h
@@ -2,6 +2,8 @@
 
 #pragma once
 
+#if PLATFORM_WINDOWS
+
 #include "Player/WmfMediaTextureSample.h"
 #include "RHI.h"
 
@@ -113,3 +115,5 @@ private:
 
 /** Implements a pool for WMF texture samples. */
 class FWmfMediaHardwareVideoDecodingTextureSamplePool : public TMediaObjectPool<FWmfMediaHardwareVideoDecodingTextureSample> { };
+
+#endif
