integrating up to CL5062415
* fix for video textures overwriting due to async rendering: member reference introduced to postpone recycling of a texture being rendered #rb:rex.hill
* FWmfMediaHardwareVideoDecodingTextureSample doesn't detect resolution change. it's a problem if it's allocated from a pool. fixed #rb:nicola.cerone

diff --git a/Engine/Plugins/Media/WmfMedia/Source/WmfMedia/Private/Wmf/WmfMediaHardwareVideoDecodingRendering.cpp b/Engine/Plugins/Media/WmfMedia/Source/WmfMedia/Private/Wmf/WmfMediaHardwareVideoDecodingRendering.cpp
--- a/Engine/Plugins/Media/WmfMedia/Source/WmfMedia/Private/Wmf/WmfMediaHardwareVideoDecodingRendering.cpp
+++ b/Engine/Plugins/Media/WmfMedia/Source/WmfMedia/Private/Wmf/WmfMediaHardwareVideoDecodingRendering.cpp
@@ -76,23 +76,7 @@ bool FWmfMediaHardwareVideoDecodingParameters::ConvertTextureFormat_RenderThread
 		GraphicsPSOInit.BoundShaderState.PixelShaderRHI = GETSAFERHISHADER_PIXEL(*PixelShader);
 		SetGraphicsPipelineState(RHICmdList, GraphicsPSOInit);
 
-		FTexture2DRHIRef SampleDestinationTexture = InSample->GetDestinationTexture();
-		if (!SampleDestinationTexture.IsValid())
-		{
-			FRHIResourceCreateInfo CreateInfo;
-			const uint32 CreateFlags = TexCreate_Dynamic | TexCreate_DisableSRVCreation;
-			FTexture2DRHIRef Texture = RHICreateTexture2D(
-				InSample->GetDim().X,
-				InSample->GetDim().Y,
-				PF_NV12,
-				1,
-				1,
-				CreateFlags,
-				CreateInfo);
-
-			InSample->SetDestinationTexture(Texture);
-			SampleDestinationTexture = Texture;
-		}
+		FTexture2DRHIRef SampleDestinationTexture = InSample->GetOrCreateDestinationTexture();
 
 		ID3D11Resource* DestinationTexture = reinterpret_cast<ID3D11Resource*>(SampleDestinationTexture->GetNativeResource());
 		if (DestinationTexture)
diff --git a/Engine/Plugins/Media/WmfMedia/Source/WmfMedia/Private/Wmf/WmfMediaHardwareVideoDecodingTextureSample.cpp b/Engine/Plugins/Media/WmfMedia/Source/WmfMedia/Private/Wmf/WmfMediaHardwareVideoDecodingTextureSample.cpp
--- a/Engine/Plugins/Media/WmfMedia/Source/WmfMedia/Private/Wmf/WmfMediaHardwareVideoDecodingTextureSample.cpp
+++ b/Engine/Plugins/Media/WmfMedia/Source/WmfMedia/Private/Wmf/WmfMediaHardwareVideoDecodingTextureSample.cpp
@@ -16,7 +16,13 @@ ID3D11Texture2D* FWmfMediaHardwareVideoDecodingTextureSample::InitializeSourceTe
 
 	if (SourceTexture.IsValid())
 	{
-		return SourceTexture;
+		D3D11_TEXTURE2D_DESC Desc;
+		SourceTexture->GetDesc(&Desc);
+
+		if (Desc.Width == Dim.X && Desc.Height == Dim.Y)
+		{
+			return SourceTexture;
+		}
 	}
 
 	D3D11_TEXTURE2D_DESC TextureDesc;
@@ -32,6 +38,7 @@ ID3D11Texture2D* FWmfMediaHardwareVideoDecodingTextureSample::InitializeSourceTe
 	TextureDesc.CPUAccessFlags = 0;
 	TextureDesc.MiscFlags = D3D11_RESOURCE_MISC_SHARED_KEYEDMUTEX;
 
+	SourceTexture.Reset();
 	InD3D11Device->CreateTexture2D(&TextureDesc, nullptr, &SourceTexture);
 
 	D3D11Device = InD3D11Device;
diff --git a/Engine/Plugins/Media/WmfMedia/Source/WmfMedia/Public/WmfMediaHardwareVideoDecodingTextureSample.h b/Engine/Plugins/Media/WmfMedia/Source/WmfMedia/Public/WmfMediaHardwareVideoDecodingTextureSample.h
--- a/Engine/Plugins/Media/WmfMedia/Source/WmfMedia/Public/WmfMediaHardwareVideoDecodingTextureSample.h
+++ b/Engine/Plugins/Media/WmfMedia/Source/WmfMedia/Public/WmfMediaHardwareVideoDecodingTextureSample.h
@@ -76,23 +76,29 @@ public:
 		return SourceTexture;
 	}
 
-	/**
-	 * Set Destination Texture of rendering device
-	 *
-	 * @param InTexture Destination texture
-	 */
-	void SetDestinationTexture(FTexture2DRHIRef InTexture)
-	{
-		DestinationTexture = InTexture;
-	}
-
 	/**
 	 * Get Destination Texture of render thread device
 	 *
 	 * @return Destination texture 
 	 */
-	FTexture2DRHIRef GetDestinationTexture() const
+	FTexture2DRHIRef GetOrCreateDestinationTexture()
 	{
+		if (DestinationTexture.IsValid() && DestinationTexture->GetSizeX() == Dim.X && DestinationTexture->GetSizeY() == Dim.Y)
+		{
+			return DestinationTexture;
+		}
+
+		FRHIResourceCreateInfo CreateInfo;
+		const uint32 CreateFlags = TexCreate_Dynamic | TexCreate_DisableSRVCreation;
+		DestinationTexture = RHICreateTexture2D(
+			Dim.X,
+			Dim.Y,
+			PF_NV12,
+			1,
+			1,
+			CreateFlags,
+			CreateInfo);
+
 		return DestinationTexture;
 	}
 
diff --git a/Engine/Source/Runtime/MediaAssets/Private/Misc/MediaTextureResource.cpp b/Engine/Source/Runtime/MediaAssets/Private/Misc/MediaTextureResource.cpp
--- a/Engine/Source/Runtime/MediaAssets/Private/Misc/MediaTextureResource.cpp
+++ b/Engine/Source/Runtime/MediaAssets/Private/Misc/MediaTextureResource.cpp
@@ -173,6 +173,8 @@ void FMediaTextureResource::Render(const FRenderParams& Params)
 {
 	check(IsInRenderingThread());
 
+	CachedSample.Reset();
+
 	SCOPE_CYCLE_COUNTER(STAT_MediaAssets_MediaTextureResourceRender);
 
 	FLinearColor Rotation(1, 0, 0, 1);
@@ -296,6 +298,10 @@ void FMediaTextureResource::Render(const FRenderParams& Params)
 			);
 		}
 #endif
+
+		// We're not done with `Sample` as rendering is asynchronous. 
+		// Hold a reference in a member to postpone recycling `Sample` till the next call
+		CachedSample = Sample;
 	}
 	else if (Params.CanClear)
 	{
diff --git a/Engine/Source/Runtime/MediaAssets/Private/Misc/MediaTextureResource.h b/Engine/Source/Runtime/MediaAssets/Private/Misc/MediaTextureResource.h
--- a/Engine/Source/Runtime/MediaAssets/Private/Misc/MediaTextureResource.h
+++ b/Engine/Source/Runtime/MediaAssets/Private/Misc/MediaTextureResource.h
@@ -177,4 +177,7 @@ private:
 
 	/** The current media player facade to get video samples from. */
 	TWeakPtr<FMediaPlayerFacade, ESPMode::ThreadSafe> PlayerFacadePtr;
+
+	/** cached media sample to postpone releasing it until the next sample rendering as it can get overwritten due to asynchronous rendering */
+	TSharedPtr<IMediaTextureSample, ESPMode::ThreadSafe> CachedSample;
 };
