Add some tempory work-arounds to get dedicated servers building again.

diff --git a/Engine/Plugins/Media/WmfMedia/Source/WmfMedia/Private/Wmf/WmfMediaHardwareVideoDecodingRendering.cpp b/Engine/Plugins/Media/WmfMedia/Source/WmfMedia/Private/Wmf/WmfMediaHardwareVideoDecodingRendering.cpp
--- a/Engine/Plugins/Media/WmfMedia/Source/WmfMedia/Private/Wmf/WmfMediaHardwareVideoDecodingRendering.cpp
+++ b/Engine/Plugins/Media/WmfMedia/Source/WmfMedia/Private/Wmf/WmfMediaHardwareVideoDecodingRendering.cpp
@@ -1,11 +1,8 @@
 // Copyright 1998-2019 Epic Games, Inc. All Rights Reserved.
 
-#include "WmfMediaHardwareVideoDecodingRendering.h"
-
-#include "WmfMediaPrivate.h"
-
-#if WMFMEDIA_SUPPORTED_PLATFORM
+#if PLATFORM_WINDOWS
 
+#include "WmfMediaHardwareVideoDecodingRendering.h"
 #include "Engine/Texture2D.h"
 #include "Engine/TextureRenderTarget2D.h"
 #include "Engine/World.h"
diff --git a/Engine/Plugins/Media/WmfMedia/Source/WmfMedia/Private/Wmf/WmfMediaHardwareVideoDecodingTextureSample.cpp b/Engine/Plugins/Media/WmfMedia/Source/WmfMedia/Private/Wmf/WmfMediaHardwareVideoDecodingTextureSample.cpp
--- a/Engine/Plugins/Media/WmfMedia/Source/WmfMedia/Private/Wmf/WmfMediaHardwareVideoDecodingTextureSample.cpp
+++ b/Engine/Plugins/Media/WmfMedia/Source/WmfMedia/Private/Wmf/WmfMediaHardwareVideoDecodingTextureSample.cpp
@@ -46,6 +46,7 @@ ID3D11Texture2D* FWmfMediaHardwareVideoDecodingTextureSample::InitializeSourceTe
 	return SourceTexture;
 }
 
+#if !WITH_SERVER_CODE
 void FWmfMediaHardwareVideoDecodingTextureSample::ShutdownPoolable()
 {
 	FWmfMediaTextureSample::ShutdownPoolable();
@@ -84,5 +85,6 @@ void FWmfMediaHardwareVideoDecodingTextureSample::ShutdownPoolable()
 		}
 	}
 }
+#endif
 
 #endif
diff --git a/Engine/Plugins/Media/WmfMedia/Source/WmfMedia/Public/WmfMediaHardwareVideoDecodingTextureSample.h b/Engine/Plugins/Media/WmfMedia/Source/WmfMedia/Public/WmfMediaHardwareVideoDecodingTextureSample.h
--- a/Engine/Plugins/Media/WmfMedia/Source/WmfMedia/Public/WmfMediaHardwareVideoDecodingTextureSample.h
+++ b/Engine/Plugins/Media/WmfMedia/Source/WmfMedia/Public/WmfMediaHardwareVideoDecodingTextureSample.h
@@ -105,7 +105,9 @@ public:
 	/**
 	 * Called the the sample is returned to the pool for cleanup purposes
 	 */
+#if !WITH_SERVER_CODE
 	virtual void ShutdownPoolable() override;
+#endif
 
 private:
 
@@ -122,4 +124,4 @@ private:
 /** Implements a pool for WMF texture samples. */
 class FWmfMediaHardwareVideoDecodingTextureSamplePool : public TMediaObjectPool<FWmfMediaHardwareVideoDecodingTextureSample> { };
 
-#endif
+#endif
\ No newline at end of file
