WMF hardware decoding, update the info in the detail to show that we are decoding the video on the hardware.

diff --git a/Engine/Plugins/Media/WmfMedia/Source/WmfMedia/Private/Player/WmfMediaTracks.cpp b/Engine/Plugins/Media/WmfMedia/Source/WmfMedia/Private/Player/WmfMediaTracks.cpp
--- a/Engine/Plugins/Media/WmfMedia/Source/WmfMedia/Private/Player/WmfMediaTracks.cpp
+++ b/Engine/Plugins/Media/WmfMedia/Source/WmfMedia/Private/Player/WmfMediaTracks.cpp
@@ -50,6 +50,7 @@ FWmfMediaTracks::FWmfMediaTracks()
 	, SelectedMetadataTrack(INDEX_NONE)
 	, SelectedVideoTrack(INDEX_NONE)
 	, SelectionChanged(false)
+	, bVideoTrackRequestedHardwareAcceleration(false)
 	, VideoSamplePool(nullptr)
 	, VideoHardwareVideoDecodingSamplePool(nullptr)
 {}
@@ -149,6 +150,7 @@ TComPtr<IMFTopology> FWmfMediaTracks::CreateTopology()
 
 	// add enabled streams to topology
 	bool TracksAdded = false;
+	bVideoTrackRequestedHardwareAcceleration = false;
 
 	if (AudioTracks.IsValidIndex(SelectedAudioTrack))
 	{
@@ -182,6 +184,7 @@ TComPtr<IMFTopology> FWmfMediaTracks::CreateTopology()
 		bool bHardwareAccelerated = false;
 		WmfMediaTopologyLoader MediaTopologyLoader;
 		bHardwareAccelerated = MediaTopologyLoader.EnableHardwareAcceleration(Topology);
+		bHardwareAccelerated = bHardwareAccelerated || bVideoTrackRequestedHardwareAcceleration;
 
 		UE_LOG(LogWmfMedia, Verbose, TEXT("Tracks %p: Video (media source %p) will be decoded on %s"), this, MediaSource.Get(), bHardwareAccelerated ? TEXT("GPU") : TEXT("CPU"));
 		Info += FString::Printf(TEXT("Video decoded on %s\n"), bHardwareAccelerated ? TEXT("GPU") : TEXT("CPU"));
@@ -1090,6 +1093,7 @@ bool FWmfMediaTracks::AddTrackToTopology(const FTrack& Track, IMFTopology& Topol
 			UE_LOG(LogWmfMedia, Verbose, TEXT("Tracks %p: Failed to configure output node for stream %i"), this, Track.StreamIndex);
 			return false;
 		}
+		bVideoTrackRequestedHardwareAcceleration = true;
 	}
 	else
 	{
diff --git a/Engine/Plugins/Media/WmfMedia/Source/WmfMedia/Private/Player/WmfMediaTracks.h b/Engine/Plugins/Media/WmfMedia/Source/WmfMedia/Private/Player/WmfMediaTracks.h
--- a/Engine/Plugins/Media/WmfMedia/Source/WmfMedia/Private/Player/WmfMediaTracks.h
+++ b/Engine/Plugins/Media/WmfMedia/Source/WmfMedia/Private/Player/WmfMediaTracks.h
@@ -332,6 +332,9 @@ private:
 	/** Whether the track selection changed. */
 	bool SelectionChanged;
 
+	/** Flag to display if the topology node has requested hardware acceleration. */
+	bool bVideoTrackRequestedHardwareAcceleration;
+
 	/** Video sample object pool. */
 	FWmfMediaTextureSamplePool* VideoSamplePool;
 	TSharedPtr<FWmfMediaHardwareVideoDecodingTextureSamplePool> VideoHardwareVideoDecodingSamplePool;
