Removing non-existent include path from WmfMedia build file (integration issue)
#rnx

diff --git a/Engine/Plugins/Media/WmfMedia/Source/WmfMedia/WmfMedia.Build.cs b/Engine/Plugins/Media/WmfMedia/Source/WmfMedia/WmfMedia.Build.cs
--- a/Engine/Plugins/Media/WmfMedia/Source/WmfMedia/WmfMedia.Build.cs
+++ b/Engine/Plugins/Media/WmfMedia/Source/WmfMedia/WmfMedia.Build.cs
@@ -60,7 +60,6 @@ namespace UnrealBuildTool.Rules
 					"WmfMedia/Private",
 					"WmfMedia/Private/Player",
 					"WmfMedia/Private/Wmf",
-                    "WmfMedia/Private/WmfMediaCodec",
                     "../../../../Source/Runtime/Windows/D3D11RHI/Private",
                     "../../../../Source/Runtime/Windows/D3D11RHI/Private/Windows",
                 });
