ElectraPlayer: updating XB1 MFSample handling to fix hang bug on close of player & related XB1 optimzations
#rnx

diff --git a/Engine/Plugins/Media/WmfMedia/Source/WmfMedia/Private/Wmf/WmfMediaHardwareVideoDecodingShaders.cpp b/Engine/Plugins/Media/WmfMedia/Source/WmfMedia/Private/Wmf/WmfMediaHardwareVideoDecodingShaders.cpp
new file mode 100644
--- /dev/null
+++ b/Engine/Plugins/Media/WmfMedia/Source/WmfMedia/Private/Wmf/WmfMediaHardwareVideoDecodingShaders.cpp
@@ -0,0 +1,8 @@
+// Copyright 1998-2019 Epic Games, Inc. All Rights Reserved.
+
+#include "WmfMediaHardwareVideoDecodingShaders.h"
+
+IMPLEMENT_SHADER_TYPE(, FHardwareVideoDecodingVS, TEXT("/Plugin/WmfMedia/Private/MediaHardwareVideoDecoding.usf"), TEXT("MainVS"), SF_Vertex)
+IMPLEMENT_SHADER_TYPE(, FHardwareVideoDecodingPS, TEXT("/Plugin/WmfMedia/Private/MediaHardwareVideoDecoding.usf"), TEXT("NV12ConvertPS"), SF_Pixel)
+IMPLEMENT_SHADER_TYPE(, FHardwareVideoDecodingPassThroughPS, TEXT("/Plugin/WmfMedia/Private/MediaHardwareVideoDecoding.usf"), TEXT("PassThroughPS"), SF_Pixel)
+IMPLEMENT_SHADER_TYPE(, FHardwareVideoDecodingY416PS, TEXT("/Plugin/WmfMedia/Private/MediaHardwareVideoDecoding.usf"), TEXT("Y416ConvertPS"), SF_Pixel)
diff --git a/Engine/Plugins/Media/WmfMedia/Source/WmfMedia/Public/WmfMediaHardwareVideoDecodingShaders.h b/Engine/Plugins/Media/WmfMedia/Source/WmfMedia/Public/WmfMediaHardwareVideoDecodingShaders.h
--- a/Engine/Plugins/Media/WmfMedia/Source/WmfMedia/Public/WmfMediaHardwareVideoDecodingShaders.h
+++ b/Engine/Plugins/Media/WmfMedia/Source/WmfMedia/Public/WmfMediaHardwareVideoDecodingShaders.h
@@ -1,5 +1,7 @@
 // Copyright 1998-2019 Epic Games, Inc. All Rights Reserved.
 
+#pragma once
+
 #include "GlobalShader.h"
 #include "RenderUtils.h"
 #include "RHIStaticStates.h"
@@ -44,14 +46,15 @@ public:
 		const TShaderRHIParamRef ShaderRHI,
 		const FShaderResourceViewRHIRef& InTextureY,
 		const FShaderResourceViewRHIRef& InTextureUV,
-		const bool InIsOutputSrgb
+		const bool InIsOutputSrgb,
+		const bool InFilterUV = true
 	)
 	{
 		SetSRVParameter(RHICmdList, ShaderRHI, TextureY, InTextureY);
 		SetSRVParameter(RHICmdList, ShaderRHI, TextureUV, InTextureUV);
 
 		SetSamplerParameter(RHICmdList, ShaderRHI, PointClampedSamplerY, TStaticSamplerState<SF_Point, AM_Clamp, AM_Clamp, AM_Clamp>::GetRHI());
-		SetSamplerParameter(RHICmdList, ShaderRHI, BilinearClampedSamplerUV, TStaticSamplerState<SF_Bilinear, AM_Clamp, AM_Clamp, AM_Clamp>::GetRHI());
+		SetSamplerParameter(RHICmdList, ShaderRHI, BilinearClampedSamplerUV, InFilterUV ? TStaticSamplerState<SF_Bilinear, AM_Clamp, AM_Clamp, AM_Clamp>::GetRHI() : TStaticSamplerState<SF_Point, AM_Clamp, AM_Clamp, AM_Clamp>::GetRHI());
 
 		SetShaderValue(RHICmdList, ShaderRHI, ColorTransform, MediaShaders::CombineColorTransformAndOffset(MediaShaders::YuvToSrgbDefault, MediaShaders::YUVOffset8bits));
 		SetShaderValue(RHICmdList, ShaderRHI, SrgbToLinear, InIsOutputSrgb ? 1 : 0); // Explicitly specify integer value, as using boolean falls over on XboxOne.
@@ -151,9 +154,3 @@ public:
 		: FWmfMediaHardwareVideoDecodingShader(Initializer)
 	{ }
 };
-
-
-IMPLEMENT_SHADER_TYPE(, FHardwareVideoDecodingVS, TEXT("/Plugin/WmfMedia/Private/MediaHardwareVideoDecoding.usf"), TEXT("MainVS"), SF_Vertex)
-IMPLEMENT_SHADER_TYPE(, FHardwareVideoDecodingPS, TEXT("/Plugin/WmfMedia/Private/MediaHardwareVideoDecoding.usf"), TEXT("NV12ConvertPS"), SF_Pixel)
-IMPLEMENT_SHADER_TYPE(, FHardwareVideoDecodingPassThroughPS, TEXT("/Plugin/WmfMedia/Private/MediaHardwareVideoDecoding.usf"), TEXT("PassThroughPS"), SF_Pixel)
-IMPLEMENT_SHADER_TYPE(, FHardwareVideoDecodingY416PS, TEXT("/Plugin/WmfMedia/Private/MediaHardwareVideoDecoding.usf"), TEXT("Y416ConvertPS"), SF_Pixel)
\ No newline at end of file
