WmfMediaPlayer: added code to avoid a deadlock during shuhtdown between tickable thread asking for playertime and player actually shutting down

diff --git a/Engine/Plugins/Media/WmfMedia/Source/WmfMedia/Private/Wmf/WmfMediaSession.cpp b/Engine/Plugins/Media/WmfMedia/Source/WmfMedia/Private/Wmf/WmfMediaSession.cpp
--- a/Engine/Plugins/Media/WmfMedia/Source/WmfMedia/Private/Wmf/WmfMediaSession.cpp
+++ b/Engine/Plugins/Media/WmfMedia/Source/WmfMedia/Private/Wmf/WmfMediaSession.cpp
@@ -55,6 +55,7 @@ FWmfMediaSession::FWmfMediaSession()
 	, SessionState(EMediaState::Closed)
 	, ShouldLoop(false)
 	, Status(EMediaStatus::None)
+	, bShuttingDown(false)
 {
 	UE_LOG(LogWmfMedia, Verbose, TEXT("Session %p: Created"), this);
 	MediaSessionCloseEvent = FPlatformProcess::GetSynchEventFromPool();
@@ -240,6 +241,7 @@ void FWmfMediaSession::Shutdown()
 		// Scope needed since MediaSession->Close() cannot be locked, see below.
 		FScopeLock Lock(&CriticalSection);
 		DiscardPendingChanges();
+		bShuttingDown = true;
 	}
 
 	// When an error occurs we close the MediaSession in HandleError(), no need to Close it again.
@@ -282,6 +284,8 @@ void FWmfMediaSession::Shutdown()
 	Status = EMediaStatus::None;
 	ThinnedRates.Empty();
 	UnthinnedRates.Empty();
+
+	bShuttingDown = false;
 }
 
 
@@ -362,6 +366,12 @@ FTimespan FWmfMediaSession::GetTime() const
 {
 	FScopeLock Lock(&CriticalSection);
 
+	if (bShuttingDown)
+	{
+		// Exit, as we could otherwise block close-event delivery while being blocked on PresentationClock calls below & keeping the CS locked
+		return FTimespan::Zero();
+	}
+
 	MFCLOCK_STATE ClockState;
 
 	if (!PresentationClock.IsValid() || FAILED(PresentationClock->GetState(0, &ClockState)) || (ClockState == MFCLOCK_STATE_INVALID))
diff --git a/Engine/Plugins/Media/WmfMedia/Source/WmfMedia/Private/Wmf/WmfMediaSession.h b/Engine/Plugins/Media/WmfMedia/Source/WmfMedia/Private/Wmf/WmfMediaSession.h
--- a/Engine/Plugins/Media/WmfMedia/Source/WmfMedia/Private/Wmf/WmfMediaSession.h
+++ b/Engine/Plugins/Media/WmfMedia/Source/WmfMedia/Private/Wmf/WmfMediaSession.h
@@ -259,6 +259,9 @@ private:
 	/** Current status flags. */
 	EMediaStatus Status;
 
+	/** Flag to signal shutdown in progress */
+	bool bShuttingDown;
+
 	/** The thinned play rates that the current media session supports. */
 	TRangeSet<float> ThinnedRates;
 
