MediaFramework: Fixed crash when playing video on Windows 7 with hardware acceleration.

diff --git a/Engine/Plugins/Media/WmfMedia/Source/WmfMedia/Private/Player/WmfMediaTracks.cpp b/Engine/Plugins/Media/WmfMedia/Source/WmfMedia/Private/Player/WmfMediaTracks.cpp
--- a/Engine/Plugins/Media/WmfMedia/Source/WmfMedia/Private/Player/WmfMediaTracks.cpp
+++ b/Engine/Plugins/Media/WmfMedia/Source/WmfMedia/Private/Player/WmfMediaTracks.cpp
@@ -1078,6 +1078,7 @@ bool FWmfMediaTracks::AddTrackToTopology(const FTrack& Track, IMFTopology& Topol
 	if ((GEngine != nullptr) && 
 		GetDefault<UWmfMediaSettings>()->HardwareAcceleratedVideoDecoding &&
 		MajorType == MFMediaType_Video &&
+		FWindowsPlatformMisc::VerifyWindowsVersion(6, 2) && // Windows 8
 		FWmfMediaStreamSink::Create(MFMediaType_Video, MediaStreamSink))
 	{
 		VideoHardwareVideoDecodingSamplePool = MakeShared<FWmfMediaHardwareVideoDecodingTextureSamplePool>();
