- Last move of some NDAd plugin files under Engine to PlatformExtensions

diff --git a/Engine/Plugins/Media/WmfMedia/Source/WmfMedia/WmfMedia.Build.cs b/Engine/Plugins/Media/WmfMedia/Source/WmfMedia/WmfMedia.Build.cs
--- a/Engine/Plugins/Media/WmfMedia/Source/WmfMedia/WmfMedia.Build.cs
+++ b/Engine/Plugins/Media/WmfMedia/Source/WmfMedia/WmfMedia.Build.cs
@@ -30,25 +30,6 @@ namespace UnrealBuildTool.Rules
                     "WmfMediaFactory"
                 });
 
-            if (Target.Platform == UnrealTargetPlatform.XboxOne)
-            {
-                PrivateDependencyModuleNames.AddRange(
-				  new string[] {
-					"D3D12RHI",
-                });
-                PrivateIncludePaths.AddRange(
-                    new string[] {
-                    "../../../../Platforms/XboxCommon/Source/Runtime/D3D12RHI/Private",
-					"../../../../Source/Runtime/D3D12RHI/Private",
-				});
-            }
-            else
-            {
-                PrivateDependencyModuleNames.AddRange(
-                  new string[] {
-                    "D3D11RHI",
-                });
-            }
 
             PrivateIncludePathModuleNames.AddRange(
 				new string[] {
@@ -77,6 +58,11 @@ namespace UnrealBuildTool.Rules
 			if ((Target.Platform == UnrealTargetPlatform.Win64) ||
 				(Target.Platform == UnrealTargetPlatform.Win32))
 			{
+				PrivateDependencyModuleNames.AddRange(
+				  new string[] {
+					"D3D11RHI",
+				});
+
 				PublicDelayLoadDLLs.Add("mf.dll");
 				PublicDelayLoadDLLs.Add("mfplat.dll");
 				PublicDelayLoadDLLs.Add("mfplay.dll");
diff --git a/Engine/Plugins/Online/VoiceChat/VivoxVoiceChat/Source/VivoxVoiceChat.Build.cs b/Engine/Plugins/Online/VoiceChat/VivoxVoiceChat/Source/VivoxVoiceChat.Build.cs
--- a/Engine/Plugins/Online/VoiceChat/VivoxVoiceChat/Source/VivoxVoiceChat.Build.cs
+++ b/Engine/Plugins/Online/VoiceChat/VivoxVoiceChat/Source/VivoxVoiceChat.Build.cs
@@ -43,15 +43,6 @@ namespace UnrealBuildTool.Rules
 					}
 				);
 			}
-			else if (Target.Platform == UnrealTargetPlatform.XboxOne)
-			{
-				PrivateDependencyModuleNames.AddRange(
-					new string[]
-					{
-						"ApplicationCore"
-					}
-				);
-			}
 		}
 	}
 }
