Fixing PVS 7.7 Issues:

warning V512: A call of the 'memcmp' function will lead to overflow of the buffer '& SubType.Data2'.

diff --git a/Engine/Plugins/Media/WmfMedia/Source/WmfMedia/Private/Wmf/WmfMediaUtils.cpp b/Engine/Plugins/Media/WmfMedia/Source/WmfMedia/Private/Wmf/WmfMediaUtils.cpp
--- a/Engine/Plugins/Media/WmfMedia/Source/WmfMedia/Private/Wmf/WmfMediaUtils.cpp
+++ b/Engine/Plugins/Media/WmfMedia/Source/WmfMedia/Private/Wmf/WmfMediaUtils.cpp
@@ -345,6 +345,11 @@ namespace WmfMedia
 	}
 
 
+	bool IsFormatFamily(const GUID& lhs, const GUID& rhs)
+	{
+		return FMemory::Memcmp(&lhs.Data2, &rhs.Data2, 12) == 0; //-V512
+	}
+
 	TComPtr<IMFMediaType> CreateOutputType(IMFMediaType& InputType, bool AllowNonStandardCodecs, bool IsVideoDevice)
 	{
 		GUID MajorType;
@@ -391,7 +396,7 @@ namespace WmfMedia
 		if (MajorType == MFMediaType_Audio)
 		{
 			// filter unsupported audio formats
-			if (FMemory::Memcmp(&SubType.Data2, &MFMPEG4Format_Base.Data2, 12) == 0)
+			if (IsFormatFamily(SubType, MFMPEG4Format_Base))
 			{
 				if (AllowNonStandardCodecs)
 				{
@@ -420,7 +425,7 @@ namespace WmfMedia
 					}
 				}
 			}
-			else if (FMemory::Memcmp(&SubType.Data2, &MFAudioFormat_Base.Data2, 12) != 0)
+			else if (!IsFormatFamily(SubType, MFAudioFormat_Base))
 			{
 				if (AllowNonStandardCodecs)
 				{
@@ -474,7 +479,7 @@ namespace WmfMedia
 		else if (MajorType == MFMediaType_Video)
 		{
 			// filter unsupported video types
-			if (FMemory::Memcmp(&SubType.Data2, &MFVideoFormat_Base.Data2, 12) != 0)
+			if (!IsFormatFamily(SubType, MFVideoFormat_Base))
 			{
 				if (AllowNonStandardCodecs)
 				{
@@ -1512,7 +1517,7 @@ namespace WmfMedia
 		if (SubType == OtherVideoFormat_LifeCam) return TEXT("LifeCam");
 		if (SubType == OtherVideoFormat_QuickTime) return TEXT("QuickTime");
 
-		if (FMemory::Memcmp(&SubType.Data2, &OtherFormatMpeg2_Base.Data2, 12) == 0)
+		if (IsFormatFamily(SubType, OtherFormatMpeg2_Base))
 		{
 			if (SubType.Data1 == OTHER_FORMAT_MPEG2_AC3) return TEXT("MPEG-2 AC3");
 			if (SubType.Data1 == OTHER_FORMAT_MPEG2_AUDIO) return TEXT("MPEG-2 Audio");
@@ -1526,8 +1531,8 @@ namespace WmfMedia
 		}
 
 		// audio formats
-		if ((FMemory::Memcmp(&SubType.Data2, &MFAudioFormat_Base.Data2, 12) == 0) ||
-			(FMemory::Memcmp(&SubType.Data2, &MFMPEG4Format_Base.Data2, 12) == 0))
+		if (IsFormatFamily(SubType, MFAudioFormat_Base) ||
+			IsFormatFamily(SubType, MFMPEG4Format_Base))
 		{
 			if (SubType.Data1 == WAVE_FORMAT_UNKNOWN) return TEXT("Unknown Audio Format");
 			if (SubType.Data1 == WAVE_FORMAT_PCM) return TEXT("PCM");
