WmfMedia: Fixed video not playing when WMF seeks but does not actually return the correct sample.
#rnx

Merging
//UE5/Release-Engine-Staging/Engine/Plugins/Media/WmfMedia/Source/WmfMedia/Private/Player/WmfMediaTracks.cpp
to //UE5/Main/Engine/Plugins/Media/WmfMedia/Source/WmfMedia/Private/Player/WmfMediaTracks.cpp

diff --git a/Engine/Plugins/Media/WmfMedia/Source/WmfMedia/Private/Player/WmfMediaTracks.cpp b/Engine/Plugins/Media/WmfMedia/Source/WmfMedia/Private/Player/WmfMediaTracks.cpp
--- a/Engine/Plugins/Media/WmfMedia/Source/WmfMedia/Private/Player/WmfMediaTracks.cpp
+++ b/Engine/Plugins/Media/WmfMedia/Source/WmfMedia/Private/Player/WmfMediaTracks.cpp
@@ -530,21 +530,29 @@ IMediaSamples::EFetchBestSampleResult FWmfMediaTracks::FetchBestVideoSampleForTi
 			// Are we waiting for the sample from a seek?
 			if (SeekTimeOptional.IsSet())
 			{
-				// Is this our seek sample?
-				FTimespan SeekTime = SeekTimeOptional.GetValue();
-				double SeekTimeSeconds = SeekTime.GetTotalSeconds();
-				if ((FMath::IsNearlyEqual(SeekTimeSeconds, SampleStartTime.GetTotalSeconds(), 0.001)) ||
-					((SeekTime >= SampleStartTime) && (SeekTime < SampleEndTime)))
+				// Are we past the seek time?
+				if (TimeRangeTime.Contains(SeekTimeOptional.GetValue()) == false)
 				{
-					// Yes this is what we have been waiting for.
-					// Reset the seek time so its no longer used.
 					SeekTimeOptional.Reset();
 				}
 				else
 				{
-					// This is not the sample we want, its old.
-					VideoSampleQueue.Pop();
-					continue;
+					// Is this our seek sample?
+					FTimespan SeekTime = SeekTimeOptional.GetValue();
+					double SeekTimeSeconds = SeekTime.GetTotalSeconds();
+					if ((FMath::IsNearlyEqual(SeekTimeSeconds, SampleStartTime.GetTotalSeconds(), 0.001)) ||
+						((SeekTime >= SampleStartTime) && (SeekTime < SampleEndTime)))
+					{
+						// Yes this is what we have been waiting for.
+						// Reset the seek time so its no longer used.
+						SeekTimeOptional.Reset();
+					}
+					else
+					{
+						// This is not the sample we want, its old.
+						VideoSampleQueue.Pop();
+						continue;
+					}
 				}
 			}
 #endif // WMFMEDIA_PLAYER_VERSION >= 2
