WmfMedia: Fixed Prores with audio not playing in DirectX 12.
#rnx

diff --git a/Engine/Plugins/Media/WmfMedia/Source/WmfMedia/Private/Player/WmfMediaTracks.cpp b/Engine/Plugins/Media/WmfMedia/Source/WmfMedia/Private/Player/WmfMediaTracks.cpp
--- a/Engine/Plugins/Media/WmfMedia/Source/WmfMedia/Private/Player/WmfMediaTracks.cpp
+++ b/Engine/Plugins/Media/WmfMedia/Source/WmfMedia/Private/Player/WmfMediaTracks.cpp
@@ -1370,13 +1370,17 @@ bool FWmfMediaTracks::AddTrackToTopology(const FTrack& Track, IMFTopology& Topol
 	}
 	else
 	{
-		// We only support our own decoders in D3D12.
-		const TCHAR* RHIName = GDynamicRHI->GetName();
-		bool bIsD3D12 = (TCString<TCHAR>::Stricmp(RHIName, TEXT("D3D12")) == 0);
-		if (bIsD3D12)
+		// Is this video?
+		if (MajorType == MFMediaType_Video)
 		{
-			UE_LOG(LogWmfMedia, Error, TEXT("Tracks %p: Format is not supported in D3D12."), this);
-			return false;
+			// We only support our own decoders in D3D12.
+			const TCHAR* RHIName = GDynamicRHI->GetName();
+			bool bIsD3D12 = (TCString<TCHAR>::Stricmp(RHIName, TEXT("D3D12")) == 0);
+			if (bIsD3D12)
+			{
+				UE_LOG(LogWmfMedia, Error, TEXT("Tracks %p: Format is not supported in D3D12."), this);
+				return false;
+			}
 		}
 
 		// connect nodes
