WmfMedia: Fix for potential crash if CopyTextureAndEnqueueSample is called after the sink has shut down.
#rnx

diff --git a/Engine/Plugins/Media/WmfMedia/Source/WmfMedia/Private/Wmf/WmfMediaStreamSink.cpp b/Engine/Plugins/Media/WmfMedia/Source/WmfMedia/Private/Wmf/WmfMediaStreamSink.cpp
--- a/Engine/Plugins/Media/WmfMedia/Source/WmfMedia/Private/Wmf/WmfMediaStreamSink.cpp
+++ b/Engine/Plugins/Media/WmfMedia/Source/WmfMedia/Private/Wmf/WmfMediaStreamSink.cpp
@@ -830,6 +830,12 @@ bool FWmfMediaStreamSink::IsVideoSampleQueueFull() const
 
 void FWmfMediaStreamSink::CopyTextureAndEnqueueSample(IMFSample* pSample)
 {
+	// CopyTextureAndEnqueueSample might get called after we have shutdown, so check first...
+	if (!CurrentMediaType.IsValid())
+	{
+		return;
+	}
+
 	UE_LOG(LogWmfMedia, VeryVerbose, TEXT("StreamSink::CopyTextureAndEnqueueSample Queue Size: %d"), VideoSampleQueue->Num());
 
 	if (IsVideoSampleQueueFull())
