Subpressing useless CA warning in WMFMediaPlayer
#rnx

diff --git a/Engine/Plugins/Media/WmfMedia/Source/WmfMedia/Private/Player/WmfMediaTracks.cpp b/Engine/Plugins/Media/WmfMedia/Source/WmfMedia/Private/Player/WmfMediaTracks.cpp
--- a/Engine/Plugins/Media/WmfMedia/Source/WmfMedia/Private/Player/WmfMediaTracks.cpp
+++ b/Engine/Plugins/Media/WmfMedia/Source/WmfMedia/Private/Player/WmfMediaTracks.cpp
@@ -1901,11 +1901,14 @@ bool FWmfMediaTracks::AddStreamToTracks(uint32 StreamIndex, bool IsVideoDevice,
 		::CoTaskMemFree(OutString);
 	}
 
+#pragma warning(push)
+#pragma warning(disable: 6388) // CA Warning: OutLength may not be NULL - According to MS documentation the initial value does not matter & we are sure to query a string type
 	if (SUCCEEDED(StreamDescriptor->GetAllocatedString(MF_SD_STREAM_NAME, &OutString, &OutLength)))
 	{
 		Track->Name = OutString;
 		::CoTaskMemFree(OutString);
 	}
+#pragma warning(pop)
 
 	Track->DisplayName = (Track->Name.IsEmpty())
 		? FText::Format(LOCTEXT("UnnamedStreamFormat", "Unnamed Track (Stream {0})"), FText::AsNumber((uint32)StreamIndex))
diff --git a/Engine/Plugins/Media/WmfMedia/Source/WmfMedia/Private/Wmf/WmfMediaUtils.cpp b/Engine/Plugins/Media/WmfMedia/Source/WmfMedia/Private/Wmf/WmfMediaUtils.cpp
--- a/Engine/Plugins/Media/WmfMedia/Source/WmfMedia/Private/Wmf/WmfMediaUtils.cpp
+++ b/Engine/Plugins/Media/WmfMedia/Source/WmfMedia/Private/Wmf/WmfMediaUtils.cpp
@@ -928,11 +928,14 @@ namespace WmfMedia
 		const GUID SymbolicLinkAttribute = IsAudioDevice ? MF_DEVSOURCE_ATTRIBUTE_SOURCE_TYPE_AUDCAP_ENDPOINT_ID : MF_DEVSOURCE_ATTRIBUTE_SOURCE_TYPE_VIDCAP_SYMBOLIC_LINK;
 		const FString UrlScheme = IsAudioDevice ? TEXT("audcap://") : TEXT("vidcap://");
 
+#pragma warning(push)
+#pragma warning(disable: 6388) // CA Warning: OutLength may not be NULL - According to MS documentation the initial value does not matter & we are sure to query a string type
 		if (SUCCEEDED(Device.GetAllocatedString(SymbolicLinkAttribute, &OutString, &OutLength)))
 		{
 			OutUrl = UrlScheme + OutString;
 			::CoTaskMemFree(OutString);
 		}
+#pragma warning(pop)
 
 		return true;
 	}
