WmfMedia: Can play in DX12 now.
#rnx

diff --git a/Engine/Plugins/Media/WmfMedia/Source/WmfMedia/Private/Player/WmfMediaTracks.cpp b/Engine/Plugins/Media/WmfMedia/Source/WmfMedia/Private/Player/WmfMediaTracks.cpp
--- a/Engine/Plugins/Media/WmfMedia/Source/WmfMedia/Private/Player/WmfMediaTracks.cpp
+++ b/Engine/Plugins/Media/WmfMedia/Source/WmfMedia/Private/Player/WmfMediaTracks.cpp
@@ -1370,17 +1370,6 @@ bool FWmfMediaTracks::AddTrackToTopology(const FTrack& Track, IMFTopology& Topol
 	}
 	else
 	{
-		// Is this video?
-		if (MajorType == MFMediaType_Video)
-		{
-			// We only support our own decoders in D3D12.
-			if (RHIGetInterfaceType() == ERHIInterfaceType::D3D12)
-			{
-				UE_LOG(LogWmfMedia, Error, TEXT("Tracks %p: Format is not supported in D3D12."), this);
-				return false;
-			}
-		}
-
 		// connect nodes
 		Result = SourceNode->ConnectOutput(0, OutputNode, 0);
 	}
