Restarts the Wmf session in the error state before switching track

diff --git a/Engine/Plugins/Media/WmfMedia/Source/WmfMedia/Private/Player/WmfMediaPlayer.cpp b/Engine/Plugins/Media/WmfMedia/Source/WmfMedia/Private/Player/WmfMediaPlayer.cpp
--- a/Engine/Plugins/Media/WmfMedia/Source/WmfMedia/Private/Player/WmfMediaPlayer.cpp
+++ b/Engine/Plugins/Media/WmfMedia/Source/WmfMedia/Private/Player/WmfMediaPlayer.cpp
@@ -181,7 +181,10 @@ void FWmfMediaPlayer::Tick()
 	if (TrackSelectionChanged)
 	{
 		// less than windows 10, seem to be a problem switching stream. The issue is also present when hardware acceleration is enabled.
-		if (!FPlatformMisc::VerifyWindowsVersion(10, 0) /* Anything < Windows 10.0 */ || GetDefault<UWmfMediaSettings>()->HardwareAcceleratedVideoDecoding)
+		// If the session is in the error state, we need to reinitialize it.
+		if (!FPlatformMisc::VerifyWindowsVersion(10, 0) /* Anything < Windows 10.0 */ ||
+			GetDefault<UWmfMediaSettings>()->HardwareAcceleratedVideoDecoding ||
+			Session->GetState() == EMediaState::Error)
 		{
 			const auto Settings = GetDefault<UWmfMediaSettings>();
 			check(Settings != nullptr);
