WmfMedia: Using hardware accelerated decoders no longer requires DefaultEngine.ini to be updated.

#rnx

diff --git a/Engine/Plugins/Media/WmfMedia/Source/WmfMedia/Private/Player/WmfMediaPlayer.cpp b/Engine/Plugins/Media/WmfMedia/Source/WmfMedia/Private/Player/WmfMediaPlayer.cpp
--- a/Engine/Plugins/Media/WmfMedia/Source/WmfMedia/Private/Player/WmfMediaPlayer.cpp
+++ b/Engine/Plugins/Media/WmfMedia/Source/WmfMedia/Private/Player/WmfMediaPlayer.cpp
@@ -184,6 +184,7 @@ void FWmfMediaPlayer::Tick()
 		// If the session is in the error state, we need to reinitialize it.
 		if (!FPlatformMisc::VerifyWindowsVersion(10, 0) /* Anything < Windows 10.0 */ ||
 			GetDefault<UWmfMediaSettings>()->HardwareAcceleratedVideoDecoding ||
+			GetDefault<UWmfMediaSettings>()->bAreHardwareAcceleratedCodecRegistered ||
 			Session->GetState() == EMediaState::Error)
 		{
 			const auto Settings = GetDefault<UWmfMediaSettings>();
diff --git a/Engine/Plugins/Media/WmfMedia/Source/WmfMedia/Private/Player/WmfMediaTracks.cpp b/Engine/Plugins/Media/WmfMedia/Source/WmfMedia/Private/Player/WmfMediaTracks.cpp
--- a/Engine/Plugins/Media/WmfMedia/Source/WmfMedia/Private/Player/WmfMediaTracks.cpp
+++ b/Engine/Plugins/Media/WmfMedia/Source/WmfMedia/Private/Player/WmfMediaTracks.cpp
@@ -181,7 +181,9 @@ TComPtr<IMFTopology> FWmfMediaTracks::CreateTopology()
 
 	UE_LOG(LogWmfMedia, Verbose, TEXT("Tracks %p: Created playback topology %p (media source %p)"), this, Topology.Get(), MediaSource.Get());
 
-	if (GetDefault<UWmfMediaSettings>()->HardwareAcceleratedVideoDecoding)
+	const UWmfMediaSettings* WmfMediaSettings = GetDefault<UWmfMediaSettings>();
+	if (WmfMediaSettings->HardwareAcceleratedVideoDecoding ||
+		WmfMediaSettings->bAreHardwareAcceleratedCodecRegistered)
 	{
 		bool bHardwareAccelerated = false;
 		WmfMediaTopologyLoader MediaTopologyLoader;
@@ -1267,8 +1269,9 @@ bool FWmfMediaTracks::AddTrackToTopology(const FTrack& Track, IMFTopology& Topol
 	// Hardware Acccelerated Stream Sink
 	TComPtr<FWmfMediaStreamSink> MediaStreamSink;
 
+	const UWmfMediaSettings* WmfMediaSettings = GetDefault<UWmfMediaSettings>();
 	if ((GEngine != nullptr) && 
-		GetDefault<UWmfMediaSettings>()->HardwareAcceleratedVideoDecoding &&
+		(WmfMediaSettings->HardwareAcceleratedVideoDecoding || WmfMediaSettings->bAreHardwareAcceleratedCodecRegistered) &&
 		MajorType == MFMediaType_Video &&
 		FPlatformMisc::VerifyWindowsVersion(6, 2) && // Windows 8
 		FWmfMediaStreamSink::Create(MFMediaType_Video, MediaStreamSink))
@@ -1527,7 +1530,9 @@ bool FWmfMediaTracks::AddStreamToTracks(uint32 StreamIndex, bool IsVideoDevice,
 	Track->SelectedFormat = INDEX_NONE;
 
 	// add track formats
-	const bool AllowNonStandardCodecs = GetDefault<UWmfMediaSettings>()->AllowNonStandardCodecs;
+	const UWmfMediaSettings* WmfMediaSettings = GetDefault<UWmfMediaSettings>();
+	const bool AllowNonStandardCodecs = WmfMediaSettings->AllowNonStandardCodecs ||
+		WmfMediaSettings->bAreHardwareAcceleratedCodecRegistered;
 
 	for (DWORD TypeIndex = 0; TypeIndex < NumMediaTypes; ++TypeIndex)
 	{
diff --git a/Engine/Plugins/Media/WmfMedia/Source/WmfMediaFactory/Public/WmfMediaSettings.cpp b/Engine/Plugins/Media/WmfMedia/Source/WmfMediaFactory/Public/WmfMediaSettings.cpp
--- a/Engine/Plugins/Media/WmfMedia/Source/WmfMediaFactory/Public/WmfMediaSettings.cpp
+++ b/Engine/Plugins/Media/WmfMedia/Source/WmfMediaFactory/Public/WmfMediaSettings.cpp
@@ -12,34 +12,5 @@ UWmfMediaSettings::UWmfMediaSettings()
 
 void UWmfMediaSettings::EnableHardwareAcceleratedCodecRegistered()
 {
-	if (bAreHardwareAcceleratedCodecRegistered == false)
-	{
-		bAreHardwareAcceleratedCodecRegistered = true;
-		AllowNonStandardCodecs = true;
-		HardwareAcceleratedVideoDecoding = true;
-#if WITH_EDITOR
-		SaveConfig(CPF_Config, *GetDefaultConfigFilename());
-#endif
-	}
+	bAreHardwareAcceleratedCodecRegistered = true;
 }
-
-#if WITH_EDITOR
-
-bool UWmfMediaSettings::CanEditChange(const FProperty* InProperty) const
-{
-	if (!Super::CanEditChange(InProperty))
-	{
-		return false;
-	}
-
-	if (InProperty->GetFName() == GET_MEMBER_NAME_CHECKED(UWmfMediaSettings, AllowNonStandardCodecs) || 
-		InProperty->GetFName() == GET_MEMBER_NAME_CHECKED(UWmfMediaSettings, HardwareAcceleratedVideoDecoding))
-	{
-		return !(bAreHardwareAcceleratedCodecRegistered);
-	}
-
-	return true;
-}
-
-#endif //WITH_EDITOR
-
diff --git a/Engine/Plugins/Media/WmfMedia/Source/WmfMediaFactory/Public/WmfMediaSettings.h b/Engine/Plugins/Media/WmfMedia/Source/WmfMediaFactory/Public/WmfMediaSettings.h
--- a/Engine/Plugins/Media/WmfMedia/Source/WmfMediaFactory/Public/WmfMediaSettings.h
+++ b/Engine/Plugins/Media/WmfMedia/Source/WmfMediaFactory/Public/WmfMediaSettings.h
@@ -62,8 +62,4 @@ public:
 
 	void EnableHardwareAcceleratedCodecRegistered();
 
-#if WITH_EDITOR
-	virtual bool CanEditChange(const FProperty* InProperty) const override;
-#endif //WITH_EDITOR
-
 };
