Adding checkes to WMF Media Player to make sure not to enable DXVA accelerated decoding when running with DX12 RHI (or Vulkan etc.)

diff --git a/Engine/Plugins/Media/WmfMedia/Source/WmfMedia/Private/Player/WmfMediaPlayer.cpp b/Engine/Plugins/Media/WmfMedia/Source/WmfMedia/Private/Player/WmfMediaPlayer.cpp
--- a/Engine/Plugins/Media/WmfMedia/Source/WmfMedia/Private/Player/WmfMediaPlayer.cpp
+++ b/Engine/Plugins/Media/WmfMedia/Source/WmfMedia/Private/Player/WmfMediaPlayer.cpp
@@ -10,6 +10,7 @@
 #include "MediaPlayerOptions.h"
 #include "Misc/Optional.h"
 #include "UObject/Class.h"
+#include "HardwareInfo.h"
 
 #include "WmfMediaSession.h"
 #include "WmfMediaSettings.h"
@@ -183,7 +184,7 @@ void FWmfMediaPlayer::Tick()
 		// less than windows 10, seem to be a problem switching stream. The issue is also present when hardware acceleration is enabled.
 		// If the session is in the error state, we need to reinitialize it.
 		if (!FPlatformMisc::VerifyWindowsVersion(10, 0) /* Anything < Windows 10.0 */ ||
-			GetDefault<UWmfMediaSettings>()->HardwareAcceleratedVideoDecoding ||
+			(GetDefault<UWmfMediaSettings>()->HardwareAcceleratedVideoDecoding && FHardwareInfo::GetHardwareInfo(NAME_RHI) == "D3D11") ||
 			GetDefault<UWmfMediaSettings>()->bAreHardwareAcceleratedCodecRegistered ||
 			Session->GetState() == EMediaState::Error)
 		{
diff --git a/Engine/Plugins/Media/WmfMedia/Source/WmfMedia/Private/Player/WmfMediaTracks.cpp b/Engine/Plugins/Media/WmfMedia/Source/WmfMedia/Private/Player/WmfMediaTracks.cpp
--- a/Engine/Plugins/Media/WmfMedia/Source/WmfMedia/Private/Player/WmfMediaTracks.cpp
+++ b/Engine/Plugins/Media/WmfMedia/Source/WmfMedia/Private/Player/WmfMediaTracks.cpp
@@ -16,6 +16,7 @@
 
 #if WITH_ENGINE
 	#include "Engine/Engine.h"
+	#include "HardwareInfo.h"
 #endif
 
 #include "Player/WmfMediaTextureSample.h"
@@ -44,7 +45,8 @@
  *****************************************************************************/
 
 FWmfMediaTracks::FWmfMediaTracks()
-	: AudioSamplePool(new FWmfMediaAudioSamplePool)
+	: bHardwareAccelerated(false)
+	, AudioSamplePool(new FWmfMediaAudioSamplePool)
 	, MediaSourceChanged(false)
 	, SelectedAudioTrack(INDEX_NONE)
 	, SelectedCaptionTrack(INDEX_NONE)
@@ -181,15 +183,22 @@ TComPtr<IMFTopology> FWmfMediaTracks::CreateTopology()
 
 	UE_LOG(LogWmfMedia, Verbose, TEXT("Tracks %p: Created playback topology %p (media source %p)"), this, Topology.Get(), MediaSource.Get());
 
+	bHardwareAccelerated = false;
+
 	const UWmfMediaSettings* WmfMediaSettings = GetDefault<UWmfMediaSettings>();
 	if (WmfMediaSettings->HardwareAcceleratedVideoDecoding ||
 		WmfMediaSettings->bAreHardwareAcceleratedCodecRegistered)
 	{
-		bool bHardwareAccelerated = false;
-		WmfMediaTopologyLoader MediaTopologyLoader;
-		bHardwareAccelerated = MediaTopologyLoader.EnableHardwareAcceleration(Topology);
-		bHardwareAccelerated = bHardwareAccelerated || bVideoTrackRequestedHardwareAcceleration;
-
+#if WITH_ENGINE
+		if (FHardwareInfo::GetHardwareInfo(NAME_RHI) == "D3D11")
+#else
+		if (1)
+#endif
+		{
+			WmfMediaTopologyLoader MediaTopologyLoader;
+			bHardwareAccelerated = MediaTopologyLoader.EnableHardwareAcceleration(Topology);
+			bHardwareAccelerated = bHardwareAccelerated || bVideoTrackRequestedHardwareAcceleration;
+		}
 		UE_LOG(LogWmfMedia, Verbose, TEXT("Tracks %p: Video (media source %p) will be decoded on %s"), this, MediaSource.Get(), bHardwareAccelerated ? TEXT("GPU") : TEXT("CPU"));
 		Info += FString::Printf(TEXT("Video decoded on %s\n"), bHardwareAccelerated ? TEXT("GPU") : TEXT("CPU"));
 	}
@@ -1271,7 +1280,7 @@ bool FWmfMediaTracks::AddTrackToTopology(const FTrack& Track, IMFTopology& Topol
 
 	const UWmfMediaSettings* WmfMediaSettings = GetDefault<UWmfMediaSettings>();
 	if ((GEngine != nullptr) && 
-		(WmfMediaSettings->HardwareAcceleratedVideoDecoding || WmfMediaSettings->bAreHardwareAcceleratedCodecRegistered) &&
+		bHardwareAccelerated &&
 		MajorType == MFMediaType_Video &&
 		FPlatformMisc::VerifyWindowsVersion(6, 2) && // Windows 8
 		FWmfMediaStreamSink::Create(MFMediaType_Video, MediaStreamSink))
diff --git a/Engine/Plugins/Media/WmfMedia/Source/WmfMedia/Private/Player/WmfMediaTracks.h b/Engine/Plugins/Media/WmfMedia/Source/WmfMedia/Private/Player/WmfMediaTracks.h
--- a/Engine/Plugins/Media/WmfMedia/Source/WmfMedia/Private/Player/WmfMediaTracks.h
+++ b/Engine/Plugins/Media/WmfMedia/Source/WmfMedia/Private/Player/WmfMediaTracks.h
@@ -301,6 +301,9 @@ private:
 
 private:
 
+	/** HW accelerated or not. */
+	bool bHardwareAccelerated;
+
 	/** Audio sample object pool. */
 	FWmfMediaAudioSamplePool* AudioSamplePool;
 
diff --git a/Engine/Plugins/Media/WmfMedia/Source/WmfMedia/Private/Wmf/WmfMediaStreamSink.cpp b/Engine/Plugins/Media/WmfMedia/Source/WmfMedia/Private/Wmf/WmfMediaStreamSink.cpp
--- a/Engine/Plugins/Media/WmfMedia/Source/WmfMedia/Private/Wmf/WmfMediaStreamSink.cpp
+++ b/Engine/Plugins/Media/WmfMedia/Source/WmfMedia/Private/Wmf/WmfMediaStreamSink.cpp
@@ -997,6 +997,8 @@ void FWmfMediaStreamSink::CopyTextureAndEnqueueSample(IMFSample* pSample)
 				PixelFormat,
 				MediaTextureSampleFormat);
 
+			if (!SharedTexture)
+				return;
 
 			D3D11_BOX SrcBox;
 			SrcBox.left = 0;
