Fix windows API define leaking out and causing function name conflict

diff --git a/Engine/Plugins/Media/WmfMedia/Source/WmfMedia/Public/WmfMediaCommon.h b/Engine/Plugins/Media/WmfMedia/Source/WmfMedia/Public/WmfMediaCommon.h
--- a/Engine/Plugins/Media/WmfMedia/Source/WmfMedia/Public/WmfMediaCommon.h
+++ b/Engine/Plugins/Media/WmfMedia/Source/WmfMedia/Public/WmfMediaCommon.h
@@ -11,6 +11,7 @@
 
 	#include "Windows/WindowsHWrapper.h"
 	#include "Windows/AllowWindowsPlatformTypes.h"
+	#include "Windows/PreWindowsApi.h"
 
 	THIRD_PARTY_INCLUDES_START
 
@@ -44,6 +45,8 @@
 	#endif
 
 	#include "Microsoft/COMPointer.h"
+	
+	#include "Windows/PostWindowsApi.h"
 	#include "Windows/HideWindowsPlatformTypes.h"
 
 #elif PLATFORM_WINDOWS && !UE_SERVER
diff --git a/Engine/Source/Runtime/Core/Public/Windows/PostWindowsApi.h b/Engine/Source/Runtime/Core/Public/Windows/PostWindowsApi.h
--- a/Engine/Source/Runtime/Core/Public/Windows/PostWindowsApi.h
+++ b/Engine/Source/Runtime/Core/Public/Windows/PostWindowsApi.h
@@ -27,6 +27,7 @@ THIRD_PARTY_INCLUDES_END
 #undef CopyFile
 #undef CreateDesktop
 #undef CreateDirectory
+#undef CreateFile
 #undef CreateFont
 #undef DeleteFile
 #undef DrawText
