Output error when media fails to open.

diff --git a/Engine/Plugins/Media/WmfMedia/Source/WmfMedia/Private/Wmf/WmfMediaSession.cpp b/Engine/Plugins/Media/WmfMedia/Source/WmfMedia/Private/Wmf/WmfMediaSession.cpp
--- a/Engine/Plugins/Media/WmfMedia/Source/WmfMedia/Private/Wmf/WmfMediaSession.cpp
+++ b/Engine/Plugins/Media/WmfMedia/Source/WmfMedia/Private/Wmf/WmfMediaSession.cpp
@@ -1478,6 +1478,9 @@ void FWmfMediaSession::HandleSessionTopologySet(HRESULT EventStatus, IMFMediaEve
 
 	if (SessionState == EMediaState::Preparing)
 	{
+		UE_LOG(LogWmfMedia, Error, TEXT("An error occurred in the media session: %s"), *WmfMedia::ResultToString(EventStatus));
+
+		// an error occurred in the topology
 		SessionState = EMediaState::Error;
 		DeferredEvents.Enqueue(EMediaEvent::MediaOpenFailed);
 	}
