Fixing up issue with H.265 being played using DX12 or Vulkan

diff --git a/Engine/Plugins/Media/WmfMedia/Source/WmfMedia/Private/Player/WmfMediaTracks.cpp b/Engine/Plugins/Media/WmfMedia/Source/WmfMedia/Private/Player/WmfMediaTracks.cpp
--- a/Engine/Plugins/Media/WmfMedia/Source/WmfMedia/Private/Player/WmfMediaTracks.cpp
+++ b/Engine/Plugins/Media/WmfMedia/Source/WmfMedia/Private/Player/WmfMediaTracks.cpp
@@ -1277,7 +1277,6 @@ bool FWmfMediaTracks::AddTrackToTopology(const FTrack& Track, IMFTopology& Topol
 
 #if WITH_ENGINE
 	if ((GEngine != nullptr) &&
-		!Track.bNoDirectXOutput &&
 		(WmfMediaSettings->HardwareAcceleratedVideoDecoding || WmfMediaSettings->bAreHardwareAcceleratedCodecRegistered) &&
 		MajorType == MFMediaType_Video &&
 		FPlatformMisc::VerifyWindowsVersion(6, 2) && // Windows 8
@@ -1581,8 +1580,7 @@ bool FWmfMediaTracks::AddStreamToTracks(uint32 StreamIndex, bool IsVideoDevice,
 		OutInfo += FString::Printf(TEXT("\t\tCodec: %s\n"), *TypeName);
 
 		// create output type
-		bool bNoDirectXOutput;
-		TComPtr<IMFMediaType> OutputType = WmfMedia::CreateOutputType(*MediaType, bNoDirectXOutput, AllowNonStandardCodecs, IsVideoDevice);
+		TComPtr<IMFMediaType> OutputType = WmfMedia::CreateOutputType(*MediaType, AllowNonStandardCodecs, IsVideoDevice);
 
 		if (!OutputType.IsValid())
 		{
@@ -1591,8 +1589,6 @@ bool FWmfMediaTracks::AddStreamToTracks(uint32 StreamIndex, bool IsVideoDevice,
 			continue;
 		}
 
-		Track->bNoDirectXOutput = bNoDirectXOutput;
-
 		// add format details
 		int32 FormatIndex = INDEX_NONE;
 
diff --git a/Engine/Plugins/Media/WmfMedia/Source/WmfMedia/Private/Player/WmfMediaTracks.h b/Engine/Plugins/Media/WmfMedia/Source/WmfMedia/Private/Player/WmfMediaTracks.h
--- a/Engine/Plugins/Media/WmfMedia/Source/WmfMedia/Private/Player/WmfMediaTracks.h
+++ b/Engine/Plugins/Media/WmfMedia/Source/WmfMedia/Private/Player/WmfMediaTracks.h
@@ -85,7 +85,6 @@ class FWmfMediaTracks
 		bool Protected;
 		int32 SelectedFormat;
 		DWORD StreamIndex;
-		bool bNoDirectXOutput;
 	};
 
 public:
diff --git a/Engine/Plugins/Media/WmfMedia/Source/WmfMedia/Private/Wmf/WmfMediaHardwareVideoDecodingRendering.cpp b/Engine/Plugins/Media/WmfMedia/Source/WmfMedia/Private/Wmf/WmfMediaHardwareVideoDecodingRendering.cpp
--- a/Engine/Plugins/Media/WmfMedia/Source/WmfMedia/Private/Wmf/WmfMediaHardwareVideoDecodingRendering.cpp
+++ b/Engine/Plugins/Media/WmfMedia/Source/WmfMedia/Private/Wmf/WmfMediaHardwareVideoDecodingRendering.cpp
@@ -118,16 +118,6 @@ bool FWmfMediaHardwareVideoDecodingParameters::ConvertTextureFormat_RenderThread
 
 	TComPtr<ID3D11Texture2D> SampleTexture = InSample->GetSourceTexture();
 
-	ID3D11Device* D3D11Device = nullptr;
-	// Must access rendering device context to copy shared resource.
-	ID3D11DeviceContext* D3D11DeviceContext = nullptr;
-
-	if (InSample->IsBufferExternal() == false)
-	{
-		D3D11Device = GetID3D11DynamicRHI()->RHIGetDevice();
-		D3D11DeviceContext = GetID3D11DynamicRHI()->RHIGetDeviceContext();
-	}
-			
 	{
 		FRHICommandListImmediate& RHICmdList = FRHICommandListExecutor::GetImmediateCommandList();
 
@@ -153,7 +143,7 @@ bool FWmfMediaHardwareVideoDecodingParameters::ConvertTextureFormat_RenderThread
 
 		// Copy buffer.
 		FTexture2DRHIRef SampleDestinationTexture = InSample->GetOrCreateDestinationTexture();
-		if (D3D11Device == nullptr)
+		if (InSample->IsBufferExternal() || !SampleTexture.IsValid())
 		{
 			const uint8* Data = (const uint8*)InSample->GetBuffer();
 			
diff --git a/Engine/Plugins/Media/WmfMedia/Source/WmfMedia/Private/Wmf/WmfMediaStreamSink.cpp b/Engine/Plugins/Media/WmfMedia/Source/WmfMedia/Private/Wmf/WmfMediaStreamSink.cpp
--- a/Engine/Plugins/Media/WmfMedia/Source/WmfMedia/Private/Wmf/WmfMediaStreamSink.cpp
+++ b/Engine/Plugins/Media/WmfMedia/Source/WmfMedia/Private/Wmf/WmfMediaStreamSink.cpp
@@ -977,64 +977,102 @@ void FWmfMediaStreamSink::CopyTextureAndEnqueueSample(IMFSample* pSample)
 		}
 		else
 		{
-			Result = pBuffer->QueryInterface(__uuidof(IMFDXGIBuffer), (LPVOID*)&pDXGIBuffer);
-			if (FAILED(Result))
+			// Only with D3D11 we can pass on data as a texture...
+			if (RHIGetInterfaceType() == ERHIInterfaceType::D3D11)
 			{
-				return;
-			}
+				Result = pBuffer->QueryInterface(__uuidof(IMFDXGIBuffer), (LPVOID*)&pDXGIBuffer);
+				if (FAILED(Result))
+				{
+					return;
+				}
 
-			Result = pDXGIBuffer->GetResource(__uuidof(ID3D11Texture2D), (LPVOID*)&pTexture2D);
-			if (FAILED(Result))
-			{
-				return;
-			}
+				Result = pDXGIBuffer->GetResource(__uuidof(ID3D11Texture2D), (LPVOID*)&pTexture2D);
+				if (FAILED(Result))
+				{
+					return;
+				}
 
-			Result = pDXGIBuffer->GetSubresourceIndex(&dwViewIndex);
-			if (FAILED(Result))
-			{
-				return;
-			}
+				Result = pDXGIBuffer->GetSubresourceIndex(&dwViewIndex);
+				if (FAILED(Result))
+				{
+					return;
+				}
 
-			check(TextureSample->GetMediaTextureSampleConverter() != nullptr);
-			ID3D11Texture2D* SharedTexture = TextureSample->InitializeSourceTexture(
-				Owner->GetDevice(),
-				FTimespan::FromMicroseconds(SampleTime / 10),
-				FTimespan::FromMicroseconds(SampleDuration / 10),
-				FIntPoint(DimX, DimY),
-				PixelFormat,
-				MediaTextureSampleFormat);
+				ID3D11Texture2D* SharedTexture = TextureSample->InitializeSourceTexture(
+					Owner->GetDevice(),
+					FTimespan::FromMicroseconds(SampleTime / 10),
+					FTimespan::FromMicroseconds(SampleDuration / 10),
+					FIntPoint(DimX, DimY),
+					PixelFormat,
+					MediaTextureSampleFormat);
 
-			if (!SharedTexture)
-			{
-				return;
-			}
+				if (!SharedTexture)
+				{
+					return;
+				}
 
-			D3D11_BOX SrcBox;
-			SrcBox.left = 0;
-			SrcBox.top = 0;
-			SrcBox.front = 0;
-			SrcBox.right = DimX;
-			SrcBox.bottom = DimY;
-			SrcBox.back = 1;
+				D3D11_BOX SrcBox;
+				SrcBox.left = 0;
+				SrcBox.top = 0;
+				SrcBox.front = 0;
+				SrcBox.right = DimX;
+				SrcBox.bottom = DimY;
+				SrcBox.back = 1;
 
-			UE_LOG(LogWmfMedia, VeryVerbose, TEXT("CopySubresourceRegion() ViewIndex:%d Time:%f"), dwViewIndex, FTimespan::FromMicroseconds(SampleTime / 10).GetTotalSeconds());
+				UE_LOG(LogWmfMedia, VeryVerbose, TEXT("CopySubresourceRegion() ViewIndex:%d Time:%f"), dwViewIndex, FTimespan::FromMicroseconds(SampleTime / 10).GetTotalSeconds());
 
-			TComPtr<IDXGIKeyedMutex> KeyedMutex;
-			SharedTexture->QueryInterface(_uuidof(IDXGIKeyedMutex), (void**)&KeyedMutex);
+				TComPtr<IDXGIKeyedMutex> KeyedMutex;
+				SharedTexture->QueryInterface(_uuidof(IDXGIKeyedMutex), (void**)&KeyedMutex);
 
-			if (KeyedMutex)
+				if (KeyedMutex)
+				{
+					// No wait on acquire since sample is new and key is 0.
+					if (KeyedMutex->AcquireSync(0, 0) == S_OK)
+					{
+						Owner->GetImmediateContext()->CopySubresourceRegion(SharedTexture, 0, 0, 0, 0, pTexture2D, dwViewIndex, &SrcBox);
+
+						// Mark texture as updated with key of 1
+						// Sample will be read in FWmfMediaHardwareVideoDecodingParameters::ConvertTextureFormat_RenderThread
+						KeyedMutex->ReleaseSync(1);
+						VideoSampleQueue->Enqueue(TextureSample);
+						UE_LOG(LogWmfMedia, VeryVerbose, TEXT("Enqueued onto VideoSampleQueue."));
+					}
+				}
+			}
+			else
 			{
-				// No wait on acquire since sample is new and key is 0.
-				if (KeyedMutex->AcquireSync(0, 0) == S_OK)
+				// Pass on sample data as CPU side buffer
+
+				DWORD BufferSize = 0;
+				if (pBuffer->GetCurrentLength(&BufferSize) != S_OK)
 				{
-					Owner->GetImmediateContext()->CopySubresourceRegion(SharedTexture, 0, 0, 0, 0, pTexture2D, dwViewIndex, &SrcBox);
+					return;
+				}
+
+				uint8* Data = nullptr;
+				if (pBuffer->Lock(&Data, NULL, NULL) == S_OK)
+				{
+					// MFW expects the sample dimension for NV12/P010 to be scaled to include "all data"
+					int SampleDimY;
+					if (PixelFormat == PF_NV12)
+					{
+						SampleDimY = (DimY * 3) / 2;
+					}
+					else
+					{
+						SampleDimY = DimY;
+					}
+
+					uint32 Pitch = (DimX / GPixelFormats[PixelFormat].BlockSizeX) * GPixelFormats[PixelFormat].BlockBytes;
+					TextureSample->Initialize(Data, BufferSize, FIntPoint(DimX, SampleDimY), FIntPoint(DimX, DimY), MediaTextureSampleFormat, Pitch, FTimespan::FromMicroseconds(SampleTime / 10), FTimespan::FromMicroseconds(SampleDuration / 10));
+					pBuffer->Unlock();
+
+					TextureSample->SetPixelFormat(PixelFormat);
 
-					// Mark texture as updated with key of 1
-					// Sample will be read in FWmfMediaHardwareVideoDecodingParameters::ConvertTextureFormat_RenderThread
-					KeyedMutex->ReleaseSync(1);
 					VideoSampleQueue->Enqueue(TextureSample);
 					UE_LOG(LogWmfMedia, VeryVerbose, TEXT("Enqueued onto VideoSampleQueue."));
 				}
+
 			}
 		}
 	}
diff --git a/Engine/Plugins/Media/WmfMedia/Source/WmfMedia/Private/Wmf/WmfMediaUtils.cpp b/Engine/Plugins/Media/WmfMedia/Source/WmfMedia/Private/Wmf/WmfMediaUtils.cpp
--- a/Engine/Plugins/Media/WmfMedia/Source/WmfMedia/Private/Wmf/WmfMediaUtils.cpp
+++ b/Engine/Plugins/Media/WmfMedia/Source/WmfMedia/Private/Wmf/WmfMediaUtils.cpp
@@ -358,7 +358,7 @@ namespace WmfMedia
 		return FMemory::Memcmp(&lhs.Data2, &rhs.Data2, 12) == 0; //-V512
 	}
 
-	TComPtr<IMFMediaType> CreateOutputType(IMFMediaType& InputType, bool& bNoDirectXOutput, bool AllowNonStandardCodecs, bool IsVideoDevice)
+	TComPtr<IMFMediaType> CreateOutputType(IMFMediaType& InputType, bool AllowNonStandardCodecs, bool IsVideoDevice)
 	{
 		GUID MajorType;
 		{
@@ -401,8 +401,6 @@ namespace WmfMedia
 			}
 		}
 
-		bNoDirectXOutput = false;
-
 		if (MajorType == MFMediaType_Audio)
 		{
 			// filter unsupported audio formats
@@ -509,9 +507,6 @@ namespace WmfMedia
 					UE_LOG(LogWmfMedia, Warning, TEXT("H264 video type requires Windows 8 or newer (your version is %s)"), *FPlatformMisc::GetOSVersion());
 					return NULL;
 				}
-
-				// Is this D3D11? Only then we can output directly into a texture
-				bNoDirectXOutput = (RHIGetInterfaceType() != ERHIInterfaceType::D3D11);
 			}
 
 			if ((SubType == MFVideoFormat_HEVC) || (SubType == MFVideoFormat_HEVC_ES))
@@ -526,9 +521,6 @@ namespace WmfMedia
 
 					UE_LOG(LogWmfMedia, Warning, TEXT("HEVC video type requires Windows 10 or newer (your version is %s), and game must be manifested for Windows 10"), *FPlatformMisc::GetOSVersion());
 				}
-
-				// Is this D3D11? Only then we can output directly into a texture
-				bNoDirectXOutput = (RHIGetInterfaceType() != ERHIInterfaceType::D3D11);
 			}
 
 			// configure video output
diff --git a/Engine/Plugins/Media/WmfMedia/Source/WmfMedia/Private/Wmf/WmfMediaUtils.h b/Engine/Plugins/Media/WmfMedia/Source/WmfMedia/Private/Wmf/WmfMediaUtils.h
--- a/Engine/Plugins/Media/WmfMedia/Source/WmfMedia/Private/Wmf/WmfMediaUtils.h
+++ b/Engine/Plugins/Media/WmfMedia/Source/WmfMedia/Private/Wmf/WmfMediaUtils.h
@@ -49,7 +49,7 @@ namespace WmfMedia
 	 * @param Whether the media source is a video capture device (there are some limitations with those).
 	 * @return The output type, or NULL on error.
 	 */
-	TComPtr<IMFMediaType> CreateOutputType(IMFMediaType& InputType, bool& bNoDirectXOutput, bool AllowNonStandardCodecs, bool IsVideoDevice);
+	TComPtr<IMFMediaType> CreateOutputType(IMFMediaType& InputType, bool AllowNonStandardCodecs, bool IsVideoDevice);
 
 	/**
 	 * Dump information about the given media type to the output log.
diff --git a/Engine/Plugins/Media/WmfMedia/Source/WmfMedia/Public/WmfMediaHardwareVideoDecodingTextureSample.h b/Engine/Plugins/Media/WmfMedia/Source/WmfMedia/Public/WmfMediaHardwareVideoDecodingTextureSample.h
--- a/Engine/Plugins/Media/WmfMedia/Source/WmfMedia/Public/WmfMediaHardwareVideoDecodingTextureSample.h
+++ b/Engine/Plugins/Media/WmfMedia/Source/WmfMedia/Public/WmfMediaHardwareVideoDecodingTextureSample.h
@@ -85,7 +85,7 @@ public:
 		// Only use sample converter for Win8+
 		if (FPlatformMisc::VerifyWindowsVersion(6, 2))
 		{
-			return this;
+			return (SourceTexture.IsValid() || DestinationTexture.IsValid() || IsBufferExternal()) ? this : nullptr;
 		}
 		return nullptr;
 	}
