Fixed bug with video decode shaders attempting to use the same sampler to multiple textures

diff --git a/Engine/Plugins/Media/WmfMedia/Shaders/Private/MediaHardwareVideoDecoding.usf b/Engine/Plugins/Media/WmfMedia/Shaders/Private/MediaHardwareVideoDecoding.usf
--- a/Engine/Plugins/Media/WmfMedia/Shaders/Private/MediaHardwareVideoDecoding.usf
+++ b/Engine/Plugins/Media/WmfMedia/Shaders/Private/MediaHardwareVideoDecoding.usf
@@ -33,6 +33,7 @@ Texture2D TextureY;
 Texture2D TextureUV;
 SamplerState PointClampedSamplerY;
 SamplerState BilinearClampedSamplerUV;
+SamplerState BilinearClampedSamplerUVAlpha;
 float4x4 ColorTransform;
 uint SrgbToLinear;
 
@@ -100,7 +101,7 @@ void YCoCgAlphaConvertPS(in noperspective float2 InUV : TEXCOORD0
 		  , out float4 OutColor : SV_Target0)
 {
 	float4 CoCgSY = TextureY.Sample(BilinearClampedSamplerUV, InUV);
-	float Alpha = TextureUV.Sample(BilinearClampedSamplerUV, InUV).x;
+	float Alpha = TextureUV.Sample(BilinearClampedSamplerUVAlpha, InUV).x;
 
 	CoCgSY += kYCoCgOffsets;
 	
diff --git a/Engine/Plugins/Media/WmfMedia/Source/WmfMedia/Public/WmfMediaHardwareVideoDecodingShaders.h b/Engine/Plugins/Media/WmfMedia/Source/WmfMedia/Public/WmfMediaHardwareVideoDecodingShaders.h
--- a/Engine/Plugins/Media/WmfMedia/Source/WmfMedia/Public/WmfMediaHardwareVideoDecodingShaders.h
+++ b/Engine/Plugins/Media/WmfMedia/Source/WmfMedia/Public/WmfMediaHardwareVideoDecodingShaders.h
@@ -37,6 +37,7 @@ public:
 
 		PointClampedSamplerY.Bind(Initializer.ParameterMap, TEXT("PointClampedSamplerY"));
 		BilinearClampedSamplerUV.Bind(Initializer.ParameterMap, TEXT("BilinearClampedSamplerUV"));
+		BilinearClampedSamplerUVAlpha.Bind(Initializer.ParameterMap, TEXT("BilinearClampedSamplerUVAlpha"));
 
 		ColorTransform.Bind(Initializer.ParameterMap, TEXT("ColorTransform"));
 		SrgbToLinear.Bind(Initializer.ParameterMap, TEXT("SrgbToLinear"));
@@ -57,6 +58,7 @@ public:
 
 		SetSamplerParameter(RHICmdList, ShaderRHI, PointClampedSamplerY, TStaticSamplerState<SF_Point, AM_Clamp, AM_Clamp, AM_Clamp>::GetRHI());
 		SetSamplerParameter(RHICmdList, ShaderRHI, BilinearClampedSamplerUV, InFilterUV ? TStaticSamplerState<SF_Bilinear, AM_Clamp, AM_Clamp, AM_Clamp>::GetRHI() : TStaticSamplerState<SF_Point, AM_Clamp, AM_Clamp, AM_Clamp>::GetRHI());
+		SetSamplerParameter(RHICmdList, ShaderRHI, BilinearClampedSamplerUVAlpha, InFilterUV ? TStaticSamplerState<SF_Bilinear, AM_Clamp, AM_Clamp, AM_Clamp>::GetRHI() : TStaticSamplerState<SF_Point, AM_Clamp, AM_Clamp, AM_Clamp>::GetRHI());
 
 		SetShaderValue(RHICmdList, ShaderRHI, ColorTransform, (FMatrix44f)MediaShaders::CombineColorTransformAndOffset(MediaShaders::YuvToRgbRec709Scaled, MediaShaders::YUVOffset8bits));
 		SetShaderValue(RHICmdList, ShaderRHI, SrgbToLinear, InIsOutputSrgb ? 1 : 0); // Explicitly specify integer value, as using boolean falls over on XboxOne.
@@ -85,6 +87,7 @@ private:
 	LAYOUT_FIELD(FShaderResourceParameter, TextureUV);
 	LAYOUT_FIELD(FShaderResourceParameter, PointClampedSamplerY);
 	LAYOUT_FIELD(FShaderResourceParameter, BilinearClampedSamplerUV);
+	LAYOUT_FIELD(FShaderResourceParameter, BilinearClampedSamplerUVAlpha);
 	LAYOUT_FIELD(FShaderParameter, ColorTransform);
 	LAYOUT_FIELD(FShaderParameter, SrgbToLinear);
 };
