WmfMedia: Fixed not looping properly.
#rnx

diff --git a/Engine/Plugins/Media/WmfMedia/Source/WmfMedia/Private/Player/WmfMediaTextureSample.h b/Engine/Plugins/Media/WmfMedia/Source/WmfMedia/Private/Player/WmfMediaTextureSample.h
--- a/Engine/Plugins/Media/WmfMedia/Source/WmfMedia/Private/Player/WmfMediaTextureSample.h
+++ b/Engine/Plugins/Media/WmfMedia/Source/WmfMedia/Private/Player/WmfMediaTextureSample.h
@@ -27,6 +27,7 @@ public:
 		, OutputDim(FIntPoint::ZeroValue)
 		, SampleFormat(EMediaTextureSampleFormat::Undefined)
 		, Time(FTimespan::Zero())
+		, SequenceIndex(0)
 		, bIsBufferExternal(false)
 		, Stride(0)
 	{ }
@@ -134,6 +135,11 @@ public:
 	 */
 	bool IsBufferExternal() const { return bIsBufferExternal; }
 
+	/**
+	 * Sets the sequence index for the time stamp of this sample.
+	 */
+	void SetSequenceIndex(int64 Index) { SequenceIndex = Index; }
+
 public:
 
 	//~ IMediaTextureSample interface
@@ -177,7 +183,7 @@ public:
 
 	virtual FMediaTimeStamp GetTime() const override
 	{
-		return FMediaTimeStamp(Time);
+		return FMediaTimeStamp(Time, SequenceIndex);
 	}
 
 	virtual bool IsCacheable() const override
@@ -216,6 +222,9 @@ protected:
 	/** Presentation for which the sample was generated. */
 	FTimespan Time;
 
+	/** Sequence index used to create the FMediaTimeStamp. */
+	int64 SequenceIndex;
+
 	/** If true, then Buffer is from WmfMediaDecoder and needs to be returned to it. */
 	bool bIsBufferExternal;
 
diff --git a/Engine/Plugins/Media/WmfMedia/Source/WmfMedia/Private/Player/WmfMediaTracks.cpp b/Engine/Plugins/Media/WmfMedia/Source/WmfMedia/Private/Player/WmfMediaTracks.cpp
--- a/Engine/Plugins/Media/WmfMedia/Source/WmfMedia/Private/Player/WmfMediaTracks.cpp
+++ b/Engine/Plugins/Media/WmfMedia/Source/WmfMedia/Private/Player/WmfMediaTracks.cpp
@@ -630,6 +630,13 @@ IMediaSamples::EFetchBestSampleResult FWmfMediaTracks::FetchBestVideoSampleForTi
 						{
 							Result = IMediaSamples::EFetchBestSampleResult::Ok;
 							CurrentOverlap = Overlap;
+							
+							// Update sequence index.
+							FWmfMediaTextureSample* WmfSample =
+								static_cast<FWmfMediaTextureSample*>(OutSample.Get());
+							WmfSample->SetSequenceIndex(GetSequenceIndex(TimeRange,
+								WmfSample->GetTime().Time));
+								
 							UE_LOG(LogWmfMedia, VeryVerbose, TEXT("FetchBestVideoSampleForTimeRange got sample."));
 						}
 					}
@@ -1993,6 +2000,33 @@ const FWmfMediaTracks::FFormat* FWmfMediaTracks::GetVideoFormat(int32 TrackIndex
 }
 
 
+int64 FWmfMediaTracks::GetSequenceIndex(const TRange<FMediaTimeStamp>& TimeRange, FTimespan Time) const
+{
+	int64 SequenceIndex = 0;
+
+	// Make sure this has the sequence index that matches the time range.
+	// If the time range only has one sequence index then just use that.
+	if (TimeRange.GetLowerBoundValue().SequenceIndex == TimeRange.GetUpperBoundValue().SequenceIndex)
+	{
+		SequenceIndex = TimeRange.GetLowerBoundValue().SequenceIndex;
+	}
+	else
+	{
+		// Try the lower bound sequence index.
+		SequenceIndex = TimeRange.GetLowerBoundValue().SequenceIndex;
+		FMediaTimeStamp SampleTime = FMediaTimeStamp(Time, SequenceIndex);
+
+		// Does our time overlap the time range?
+		if (TimeRange.Contains(SampleTime) == false)
+		{
+			// No. Use the upper bound sequence index.
+			SequenceIndex = TimeRange.GetUpperBoundValue().SequenceIndex;
+		}
+	}
+
+	return SequenceIndex;
+}
+
 /* FWmfMediaTracks callbacks
  *****************************************************************************/
 
diff --git a/Engine/Plugins/Media/WmfMedia/Source/WmfMedia/Private/Player/WmfMediaTracks.h b/Engine/Plugins/Media/WmfMedia/Source/WmfMedia/Private/Player/WmfMediaTracks.h
--- a/Engine/Plugins/Media/WmfMedia/Source/WmfMedia/Private/Player/WmfMediaTracks.h
+++ b/Engine/Plugins/Media/WmfMedia/Source/WmfMedia/Private/Player/WmfMediaTracks.h
@@ -282,6 +282,15 @@ private:
 	 */
 	const FFormat* GetVideoFormat(int32 TrackIndex, int32 FormatIndex) const;
 
+	/**
+	 * Gets the sequence index for a specific time that fits the time range.
+	 * 
+	 * @param	TimeRange		Sequence index will be in this range.
+	 * @param	Time			Time to retrieve sequence index for.
+	 * @return					Sequence index.
+	 */
+	int64 GetSequenceIndex(const TRange<FMediaTimeStamp>& TimeRange, FTimespan Time) const;
+
 private:
 
 	/** Callback for handling media sampler pauses. */
