Changing WMF Media shaders to use parameter structs.

diff --git a/Engine/Plugins/Media/WmfMedia/Source/WmfMedia/Private/Wmf/WmfMediaHardwareVideoDecodingRendering.cpp b/Engine/Plugins/Media/WmfMedia/Source/WmfMedia/Private/Wmf/WmfMediaHardwareVideoDecodingRendering.cpp
--- a/Engine/Plugins/Media/WmfMedia/Source/WmfMedia/Private/Wmf/WmfMediaHardwareVideoDecodingRendering.cpp
+++ b/Engine/Plugins/Media/WmfMedia/Source/WmfMedia/Private/Wmf/WmfMediaHardwareVideoDecodingRendering.cpp
@@ -167,7 +167,6 @@ bool FWmfMediaHardwareVideoDecodingParameters::ConvertTextureFormat_RenderThread
 
 		FGlobalShaderMap* GlobalShaderMap = GetGlobalShaderMap(GMaxRHIFeatureLevel);
 
-
 		if (InSample->GetFormat() == EMediaTextureSampleFormat::CharNV12)
 		{
 			TShaderMapRef< FHardwareVideoDecodingVS > VertexShader(GlobalShaderMap);
@@ -179,8 +178,11 @@ bool FWmfMediaHardwareVideoDecodingParameters::ConvertTextureFormat_RenderThread
 
 			FShaderResourceViewRHIRef Y_SRV = RHICreateShaderResourceView(SampleDestinationTexture, 0, 1, PF_G8);
 			FShaderResourceViewRHIRef UV_SRV = RHICreateShaderResourceView(SampleDestinationTexture, 0, 1, PF_R8G8);
-			VertexShader->SetParameters(RHICmdList, VertexShader.GetVertexShader(), Y_SRV, UV_SRV, InSample->IsOutputSrgb());
-			PixelShader->SetParameters(RHICmdList, PixelShader.GetPixelShader(), Y_SRV, UV_SRV, InSample->IsOutputSrgb());
+
+			const FHardwareVideoDecodingShaderParams Parameters = FWmfMediaHardwareVideoDecodingShader::GetParameters(Y_SRV, UV_SRV, InSample->IsOutputSrgb());
+
+			SetShaderParameters(RHICmdList, VertexShader, VertexShader.GetVertexShader(), Parameters);
+			SetShaderParameters(RHICmdList, PixelShader, VertexShader.GetPixelShader(), Parameters);
 		}
 		else if ((InSample->GetFormat() == EMediaTextureSampleFormat::CharBGRA) ||
 			(InSample->GetFormat() == EMediaTextureSampleFormat::DXT1) ||
@@ -203,8 +205,11 @@ bool FWmfMediaHardwareVideoDecodingParameters::ConvertTextureFormat_RenderThread
 				PixelFormat = PF_DXT5;
 			}
 			FShaderResourceViewRHIRef SRV = RHICreateShaderResourceView(SampleDestinationTexture, 0, 1, PixelFormat);
-			VertexShader->SetParameters(RHICmdList, VertexShader.GetVertexShader(), SRV, InSample->IsOutputSrgb());
-			PixelShader->SetParameters(RHICmdList, PixelShader.GetPixelShader(), SRV, InSample->IsOutputSrgb());
+
+			const FHardwareVideoDecodingShaderParams Parameters = FWmfMediaHardwareVideoDecodingShader::GetParameters(SRV, InSample->IsOutputSrgb());
+
+			SetShaderParameters(RHICmdList, VertexShader, VertexShader.GetVertexShader(), Parameters);
+			SetShaderParameters(RHICmdList, PixelShader, PixelShader.GetPixelShader(), Parameters);
 		}
 		else if (InSample->GetFormat() == EMediaTextureSampleFormat::Y416)
 		{
@@ -216,8 +221,11 @@ bool FWmfMediaHardwareVideoDecodingParameters::ConvertTextureFormat_RenderThread
 			SetGraphicsPipelineState(RHICmdList, GraphicsPSOInit, 0);
 
 			FShaderResourceViewRHIRef SRV = RHICreateShaderResourceView(SampleDestinationTexture, 0, 1, PF_A16B16G16R16);
-			VertexShader->SetParameters(RHICmdList, VertexShader.GetVertexShader(), SRV, InSample->IsOutputSrgb());
-			PixelShader->SetParameters(RHICmdList, PixelShader.GetPixelShader(), SRV, InSample->IsOutputSrgb());
+
+			const FHardwareVideoDecodingShaderParams Parameters = FWmfMediaHardwareVideoDecodingShader::GetParameters(SRV, InSample->IsOutputSrgb());
+
+			SetShaderParameters(RHICmdList, VertexShader, VertexShader.GetVertexShader(), Parameters);
+			SetShaderParameters(RHICmdList, PixelShader, PixelShader.GetPixelShader(), Parameters);
 		}
 		else if (InSample->GetFormat() == EMediaTextureSampleFormat::YCoCg_DXT5)
 		{
@@ -229,9 +237,11 @@ bool FWmfMediaHardwareVideoDecodingParameters::ConvertTextureFormat_RenderThread
 			SetGraphicsPipelineState(RHICmdList, GraphicsPSOInit, 0);
 
 			FShaderResourceViewRHIRef SRV = RHICreateShaderResourceView(SampleDestinationTexture, 0, 1, PF_DXT5);
-			VertexShader->SetParameters(RHICmdList, VertexShader.GetVertexShader(), SRV, InSample->IsOutputSrgb());
-			PixelShader->SetParameters(RHICmdList, PixelShader.GetPixelShader(), SRV, InSample->IsOutputSrgb());
 
+			const FHardwareVideoDecodingShaderParams Parameters = FWmfMediaHardwareVideoDecodingShader::GetParameters(SRV, InSample->IsOutputSrgb());
+
+			SetShaderParameters(RHICmdList, VertexShader, VertexShader.GetVertexShader(), Parameters);
+			SetShaderParameters(RHICmdList, PixelShader, PixelShader.GetPixelShader(), Parameters);
 		}
 		else if (InSample->GetFormat() == EMediaTextureSampleFormat::YCoCg_DXT5_Alpha_BC4)
 		{
@@ -255,9 +265,12 @@ bool FWmfMediaHardwareVideoDecodingParameters::ConvertTextureFormat_RenderThread
 
 			FShaderResourceViewRHIRef SRV = RHICreateShaderResourceView(SampleDestinationTexture, 0, 1, PF_DXT5);
 			FShaderResourceViewRHIRef SRVAlpha = RHICreateShaderResourceView(SampleDestinationAlphaTexture, 0, 1, PF_BC4);
-			VertexShader->SetParameters(RHICmdList, VertexShader.GetVertexShader(), SRV, InSample->IsOutputSrgb());
-			PixelShader->SetParameters(RHICmdList, PixelShader.GetPixelShader(), SRV, SRVAlpha, InSample->IsOutputSrgb());
 
+			const FHardwareVideoDecodingShaderParams VsParameters = FWmfMediaHardwareVideoDecodingShader::GetParameters(SRV, InSample->IsOutputSrgb());
+			const FHardwareVideoDecodingShaderParams PsParameters = FWmfMediaHardwareVideoDecodingShader::GetParameters(SRV, SRVAlpha, InSample->IsOutputSrgb());
+
+			SetShaderParameters(RHICmdList, VertexShader, VertexShader.GetVertexShader(), VsParameters);
+			SetShaderParameters(RHICmdList, PixelShader, PixelShader.GetPixelShader(), PsParameters);
 		}
 
 		RHICmdList.DrawPrimitive(0, 2, 1);
diff --git a/Engine/Plugins/Media/WmfMedia/Source/WmfMedia/Private/Wmf/WmfMediaHardwareVideoDecodingShaders.cpp b/Engine/Plugins/Media/WmfMedia/Source/WmfMedia/Private/Wmf/WmfMediaHardwareVideoDecodingShaders.cpp
--- a/Engine/Plugins/Media/WmfMedia/Source/WmfMedia/Private/Wmf/WmfMediaHardwareVideoDecodingShaders.cpp
+++ b/Engine/Plugins/Media/WmfMedia/Source/WmfMedia/Private/Wmf/WmfMediaHardwareVideoDecodingShaders.cpp
@@ -2,8 +2,6 @@
 
 #include "WmfMediaHardwareVideoDecodingShaders.h"
 
-IMPLEMENT_TYPE_LAYOUT(FWmfMediaHardwareVideoDecodingShader);
-
 IMPLEMENT_SHADER_TYPE(, FHardwareVideoDecodingVS, TEXT("/Plugin/WmfMedia/Private/MediaHardwareVideoDecoding.usf"), TEXT("MainVS"), SF_Vertex)
 IMPLEMENT_SHADER_TYPE(, FHardwareVideoDecodingPS, TEXT("/Plugin/WmfMedia/Private/MediaHardwareVideoDecoding.usf"), TEXT("NV12ConvertPS"), SF_Pixel)
 IMPLEMENT_SHADER_TYPE(, FHardwareVideoDecodingPassThroughPS, TEXT("/Plugin/WmfMedia/Private/MediaHardwareVideoDecoding.usf"), TEXT("PassThroughPS"), SF_Pixel)
diff --git a/Engine/Plugins/Media/WmfMedia/Source/WmfMedia/Public/WmfMediaHardwareVideoDecodingShaders.h b/Engine/Plugins/Media/WmfMedia/Source/WmfMedia/Public/WmfMediaHardwareVideoDecodingShaders.h
--- a/Engine/Plugins/Media/WmfMedia/Source/WmfMedia/Public/WmfMediaHardwareVideoDecodingShaders.h
+++ b/Engine/Plugins/Media/WmfMedia/Source/WmfMedia/Public/WmfMediaHardwareVideoDecodingShaders.h
@@ -10,13 +10,32 @@
 #include "MediaShaders.h"
 #include "DataDrivenShaderPlatformInfo.h"
 
+BEGIN_SHADER_PARAMETER_STRUCT(FHardwareVideoDecodingShaderParams, )
+	SHADER_PARAMETER_SRV(Texture2D, TextureY)
+	SHADER_PARAMETER_SRV(Texture2D, TextureUV)
+
+	SHADER_PARAMETER_SAMPLER(SamplerState, PointClampedSamplerY)
+	SHADER_PARAMETER_SAMPLER(SamplerState, BilinearClampedSamplerUV)
+	SHADER_PARAMETER_SAMPLER(SamplerState, BilinearClampedSamplerUVAlpha)
+
+	SHADER_PARAMETER(FMatrix44f, ColorTransform)
+	SHADER_PARAMETER(uint32, SrgbToLinear)
+END_SHADER_PARAMETER_STRUCT()
+
 /**
  * Shaders which allow the conversion of NV12 texture data into RGB textures.
  */
 class FWmfMediaHardwareVideoDecodingShader : public FGlobalShader
 {
-	DECLARE_TYPE_LAYOUT(FWmfMediaHardwareVideoDecodingShader, NonVirtual);
 public:
+	using FParameters = FHardwareVideoDecodingShaderParams;
+
+	FWmfMediaHardwareVideoDecodingShader() = default;
+	FWmfMediaHardwareVideoDecodingShader(const ShaderMetaType::CompiledShaderInitializerType& Initializer)
+		: FGlobalShader(Initializer)
+	{
+	}
+
 	static bool ShouldCompilePermutation(const FGlobalShaderPermutationParameters& Parameters)
 	{
 		return IsFeatureLevelSupported(Parameters.Platform, ERHIFeatureLevel::SM5) 
@@ -28,162 +47,90 @@ public:
 		FGlobalShader::ModifyCompilationEnvironment(Parameters, OutEnvironment);
 	}
 
-	FWmfMediaHardwareVideoDecodingShader() {}
-
-	FWmfMediaHardwareVideoDecodingShader(const ShaderMetaType::CompiledShaderInitializerType& Initializer)
-		: FGlobalShader(Initializer)
+	static FHardwareVideoDecodingShaderParams GetCommonParameters(bool InIsOutputSrgb)
 	{
-		TextureY.Bind(Initializer.ParameterMap, TEXT("TextureY"));
-		TextureUV.Bind(Initializer.ParameterMap, TEXT("TextureUV"));
+		FHardwareVideoDecodingShaderParams Parameters{};
+
+		Parameters.PointClampedSamplerY = TStaticSamplerState<SF_Point, AM_Clamp, AM_Clamp, AM_Clamp>::GetRHI();
 
-		PointClampedSamplerY.Bind(Initializer.ParameterMap, TEXT("PointClampedSamplerY"));
-		BilinearClampedSamplerUV.Bind(Initializer.ParameterMap, TEXT("BilinearClampedSamplerUV"));
-		BilinearClampedSamplerUVAlpha.Bind(Initializer.ParameterMap, TEXT("BilinearClampedSamplerUVAlpha"));
+		Parameters.ColorTransform = (FMatrix44f)MediaShaders::CombineColorTransformAndOffset(MediaShaders::YuvToRgbRec709Scaled, MediaShaders::YUVOffset8bits);
 
-		ColorTransform.Bind(Initializer.ParameterMap, TEXT("ColorTransform"));
-		SrgbToLinear.Bind(Initializer.ParameterMap, TEXT("SrgbToLinear"));
+		// Explicitly specify integer value, as using boolean falls over on some platforms.
+		Parameters.SrgbToLinear = InIsOutputSrgb ? 1 : 0;
+		return Parameters;
 	}
 
-	template<typename TShaderRHIParamRef>
-	void SetParameters(
-		FRHICommandListImmediate& RHICmdList,
-		const TShaderRHIParamRef ShaderRHI,
-		const FShaderResourceViewRHIRef& InTextureY,
-		const FShaderResourceViewRHIRef& InTextureUV,
-		const bool InIsOutputSrgb,
-		const bool InFilterUV = true
-	)
+	static FHardwareVideoDecodingShaderParams GetParameters(FRHIShaderResourceView* InTextureY, FRHIShaderResourceView* InTextureUV, bool InIsOutputSrgb, bool InFilterUV = true)
 	{
-		SetSRVParameter(RHICmdList, ShaderRHI, TextureY, InTextureY);
-		SetSRVParameter(RHICmdList, ShaderRHI, TextureUV, InTextureUV);
-
-		SetSamplerParameter(RHICmdList, ShaderRHI, PointClampedSamplerY, TStaticSamplerState<SF_Point, AM_Clamp, AM_Clamp, AM_Clamp>::GetRHI());
-		SetSamplerParameter(RHICmdList, ShaderRHI, BilinearClampedSamplerUV, InFilterUV ? TStaticSamplerState<SF_Bilinear, AM_Clamp, AM_Clamp, AM_Clamp>::GetRHI() : TStaticSamplerState<SF_Point, AM_Clamp, AM_Clamp, AM_Clamp>::GetRHI());
-		SetSamplerParameter(RHICmdList, ShaderRHI, BilinearClampedSamplerUVAlpha, InFilterUV ? TStaticSamplerState<SF_Bilinear, AM_Clamp, AM_Clamp, AM_Clamp>::GetRHI() : TStaticSamplerState<SF_Point, AM_Clamp, AM_Clamp, AM_Clamp>::GetRHI());
-
-		SetShaderValue(RHICmdList, ShaderRHI, ColorTransform, (FMatrix44f)MediaShaders::CombineColorTransformAndOffset(MediaShaders::YuvToRgbRec709Scaled, MediaShaders::YUVOffset8bits));
-		SetShaderValue(RHICmdList, ShaderRHI, SrgbToLinear, InIsOutputSrgb ? 1 : 0); // Explicitly specify integer value, as using boolean falls over on XboxOne.
+		FHardwareVideoDecodingShaderParams Parameters = GetCommonParameters(InIsOutputSrgb);
+
+		Parameters.TextureY = InTextureY;
+		Parameters.TextureUV = InTextureUV;
+
+		if (InFilterUV)
+		{
+			Parameters.BilinearClampedSamplerUV = TStaticSamplerState<SF_Bilinear, AM_Clamp, AM_Clamp, AM_Clamp>::GetRHI();
+			Parameters.BilinearClampedSamplerUVAlpha = TStaticSamplerState<SF_Bilinear, AM_Clamp, AM_Clamp, AM_Clamp>::GetRHI();
+		}
+		else
+		{
+			Parameters.BilinearClampedSamplerUV = TStaticSamplerState<SF_Point, AM_Clamp, AM_Clamp, AM_Clamp>::GetRHI();
+			Parameters.BilinearClampedSamplerUVAlpha = TStaticSamplerState<SF_Point, AM_Clamp, AM_Clamp, AM_Clamp>::GetRHI();
+		}
+		return Parameters;
 	}
 
-	template<typename TShaderRHIParamRef>
-	void SetParameters(
-		FRHICommandListImmediate& RHICmdList,
-		const TShaderRHIParamRef ShaderRHI,
-		const FShaderResourceViewRHIRef& InTextureRGBA,
-		const bool InIsOutputSrgb
-	)
+	static FHardwareVideoDecodingShaderParams GetParameters(FRHIShaderResourceView* InTextureRGBA, bool InIsOutputSrgb)
 	{
-		SetSRVParameter(RHICmdList, ShaderRHI, TextureY, InTextureRGBA);
-
-		SetSamplerParameter(RHICmdList, ShaderRHI, PointClampedSamplerY, TStaticSamplerState<SF_Point, AM_Clamp, AM_Clamp, AM_Clamp>::GetRHI());
-		SetSamplerParameter(RHICmdList, ShaderRHI, BilinearClampedSamplerUV, TStaticSamplerState<SF_Bilinear, AM_Clamp, AM_Clamp, AM_Clamp>::GetRHI());
-
-		SetShaderValue(RHICmdList, ShaderRHI, ColorTransform, (FMatrix44f)MediaShaders::CombineColorTransformAndOffset(MediaShaders::YuvToRgbRec709Scaled, MediaShaders::YUVOffset8bits));
-		SetShaderValue(RHICmdList, ShaderRHI, SrgbToLinear, InIsOutputSrgb ? 1 : 0); // Explicitly specify integer value, as using boolean falls over on XboxOne.
+		FHardwareVideoDecodingShaderParams Parameters = GetCommonParameters(InIsOutputSrgb);
+		Parameters.TextureY = InTextureRGBA;
+		Parameters.BilinearClampedSamplerUV = TStaticSamplerState<SF_Bilinear, AM_Clamp, AM_Clamp, AM_Clamp>::GetRHI();
+		return Parameters;
 	}
-
-private:
-	
-	LAYOUT_FIELD(FShaderResourceParameter, TextureY);
-	LAYOUT_FIELD(FShaderResourceParameter, TextureUV);
-	LAYOUT_FIELD(FShaderResourceParameter, PointClampedSamplerY);
-	LAYOUT_FIELD(FShaderResourceParameter, BilinearClampedSamplerUV);
-	LAYOUT_FIELD(FShaderResourceParameter, BilinearClampedSamplerUVAlpha);
-	LAYOUT_FIELD(FShaderParameter, ColorTransform);
-	LAYOUT_FIELD(FShaderParameter, SrgbToLinear);
 };
 
 class FHardwareVideoDecodingVS : public FWmfMediaHardwareVideoDecodingShader
 {
-	DECLARE_SHADER_TYPE(FHardwareVideoDecodingVS, Global);
-
 public:
-
-	/** Default constructor. */
-	FHardwareVideoDecodingVS() {}
-
-	/** Initialization constructor. */
-	FHardwareVideoDecodingVS(const ShaderMetaType::CompiledShaderInitializerType& Initializer)
-		: FWmfMediaHardwareVideoDecodingShader(Initializer)
-	{
-	}
+	DECLARE_GLOBAL_SHADER(FHardwareVideoDecodingVS);
+	SHADER_USE_PARAMETER_STRUCT(FHardwareVideoDecodingVS, FWmfMediaHardwareVideoDecodingShader);
 };
 
 class FHardwareVideoDecodingPS : public FWmfMediaHardwareVideoDecodingShader
 {
-	DECLARE_SHADER_TYPE(FHardwareVideoDecodingPS, Global);
-
 public:
-
-	/** Default constructor. */
-	FHardwareVideoDecodingPS() {}
-
-	/** Initialization constructor. */
-	FHardwareVideoDecodingPS(const ShaderMetaType::CompiledShaderInitializerType& Initializer)
-		: FWmfMediaHardwareVideoDecodingShader(Initializer)
-	{ }
+	DECLARE_GLOBAL_SHADER(FHardwareVideoDecodingPS);
+	SHADER_USE_PARAMETER_STRUCT(FHardwareVideoDecodingPS, FWmfMediaHardwareVideoDecodingShader);
 };
 
 class FHardwareVideoDecodingPassThroughPS : public FWmfMediaHardwareVideoDecodingShader
 {
-	DECLARE_SHADER_TYPE(FHardwareVideoDecodingPassThroughPS, Global);
-
 public:
-
-	/** Default constructor. */
-	FHardwareVideoDecodingPassThroughPS() {}
-
-	/** Initialization constructor. */
-	FHardwareVideoDecodingPassThroughPS(const ShaderMetaType::CompiledShaderInitializerType& Initializer)
-		: FWmfMediaHardwareVideoDecodingShader(Initializer)
-	{ }
+	DECLARE_GLOBAL_SHADER(FHardwareVideoDecodingPassThroughPS);
+	SHADER_USE_PARAMETER_STRUCT(FHardwareVideoDecodingPassThroughPS, FWmfMediaHardwareVideoDecodingShader);
 };
 
 class FHardwareVideoDecodingY416PS : public FWmfMediaHardwareVideoDecodingShader
 {
-	DECLARE_SHADER_TYPE(FHardwareVideoDecodingY416PS, Global);
-
 public:
-
-	/** Default constructor. */
-	FHardwareVideoDecodingY416PS() {}
-
-	/** Initialization constructor. */
-	FHardwareVideoDecodingY416PS(const ShaderMetaType::CompiledShaderInitializerType& Initializer)
-		: FWmfMediaHardwareVideoDecodingShader(Initializer)
-	{ }
+	DECLARE_GLOBAL_SHADER(FHardwareVideoDecodingY416PS);
+	SHADER_USE_PARAMETER_STRUCT(FHardwareVideoDecodingY416PS, FWmfMediaHardwareVideoDecodingShader);
 };
 
 
 class FHardwareVideoDecodingYCoCgPS : public FWmfMediaHardwareVideoDecodingShader
 {
-	DECLARE_SHADER_TYPE(FHardwareVideoDecodingYCoCgPS, Global);
-
 public:
-
-	/** Default constructor. */
-	FHardwareVideoDecodingYCoCgPS() {}
-
-	/** Initialization constructor. */
-	FHardwareVideoDecodingYCoCgPS(const ShaderMetaType::CompiledShaderInitializerType& Initializer)
-		: FWmfMediaHardwareVideoDecodingShader(Initializer)
-	{ }
+	DECLARE_GLOBAL_SHADER(FHardwareVideoDecodingYCoCgPS);
+	SHADER_USE_PARAMETER_STRUCT(FHardwareVideoDecodingYCoCgPS, FWmfMediaHardwareVideoDecodingShader);
 };
 
 
 class FHardwareVideoDecodingYCoCgAlphaPS : public FWmfMediaHardwareVideoDecodingShader
 {
-	DECLARE_SHADER_TYPE(FHardwareVideoDecodingYCoCgAlphaPS, Global);
-
 public:
-
-	/** Default constructor. */
-	FHardwareVideoDecodingYCoCgAlphaPS() {}
-
-	/** Initialization constructor. */
-	FHardwareVideoDecodingYCoCgAlphaPS(const ShaderMetaType::CompiledShaderInitializerType& Initializer)
-		: FWmfMediaHardwareVideoDecodingShader(Initializer)
-	{ }
+	DECLARE_GLOBAL_SHADER(FHardwareVideoDecodingYCoCgAlphaPS);
+	SHADER_USE_PARAMETER_STRUCT(FHardwareVideoDecodingYCoCgAlphaPS, FWmfMediaHardwareVideoDecodingShader);
 };
 
 
