[Engine/Plugins]
* Fixed two non unity errors that the farm reported

diff --git a/Engine/Plugins/Experimental/NNE/Source/NNXCore/Private/NNECoreTypes.cpp b/Engine/Plugins/Experimental/NNE/Source/NNXCore/Private/NNECoreTypes.cpp
--- a/Engine/Plugins/Experimental/NNE/Source/NNXCore/Private/NNECoreTypes.cpp
+++ b/Engine/Plugins/Experimental/NNE/Source/NNXCore/Private/NNECoreTypes.cpp
@@ -1,7 +1,5 @@
 // Copyright Epic Games, Inc. All Rights Reserved.
 
-#pragma once
-
 #include "NNECoreTypes.h"
 #include "NNXCore.h"
 
diff --git a/Engine/Plugins/Media/WmfMedia/Source/WmfMediaFactory/Private/WmfMediaFactoryModule.cpp b/Engine/Plugins/Media/WmfMedia/Source/WmfMediaFactory/Private/WmfMediaFactoryModule.cpp
--- a/Engine/Plugins/Media/WmfMedia/Source/WmfMediaFactory/Private/WmfMediaFactoryModule.cpp
+++ b/Engine/Plugins/Media/WmfMedia/Source/WmfMediaFactory/Private/WmfMediaFactoryModule.cpp
@@ -15,6 +15,8 @@
 #if WITH_EDITOR
 	#include "ISettingsModule.h"
 	#include "WmfMediaSettings.h"
+#else
+	#include "Internationalization/Internationalization.h"
 #endif
 
 #include "../../WmfMedia/Public/IWmfMediaModule.h"
