Exclude Wmf media shaders from GLES31 desktop platform

diff --git a/Engine/Plugins/Media/WmfMedia/Source/WmfMedia/Public/WmfMediaHardwareVideoDecodingShaders.h b/Engine/Plugins/Media/WmfMedia/Source/WmfMedia/Public/WmfMediaHardwareVideoDecodingShaders.h
--- a/Engine/Plugins/Media/WmfMedia/Source/WmfMedia/Public/WmfMediaHardwareVideoDecodingShaders.h
+++ b/Engine/Plugins/Media/WmfMedia/Source/WmfMedia/Public/WmfMediaHardwareVideoDecodingShaders.h
@@ -38,8 +38,8 @@ public:
 
 	static bool ShouldCompilePermutation(const FGlobalShaderPermutationParameters& Parameters)
 	{
-		return IsFeatureLevelSupported(Parameters.Platform, ERHIFeatureLevel::SM5) 
-			|| IsPCPlatform(Parameters.Platform); // to support mobile emulation on PC
+		return IsFeatureLevelSupported(Parameters.Platform, ERHIFeatureLevel::SM5)
+			|| (IsPCPlatform(Parameters.Platform) && !IsOpenGLPlatform(Parameters.Platform)); // to support mobile emulation on PC
 	}
 
 	static void ModifyCompilationEnvironment(const FGlobalShaderPermutationParameters& Parameters, FShaderCompilerEnvironment& OutEnvironment)
