Fix -Wshadow clang error

#rnx

diff --git a/Engine/Plugins/Media/WmfMedia/Source/WmfMedia/Private/WmfMediaCodec/WmfMediaDecoder.cpp b/Engine/Plugins/Media/WmfMedia/Source/WmfMedia/Private/WmfMediaCodec/WmfMediaDecoder.cpp
--- a/Engine/Plugins/Media/WmfMedia/Source/WmfMedia/Private/WmfMediaCodec/WmfMediaDecoder.cpp
+++ b/Engine/Plugins/Media/WmfMedia/Source/WmfMedia/Private/WmfMediaCodec/WmfMediaDecoder.cpp
@@ -153,9 +153,9 @@ HRESULT WmfMediaDecoder::GetAttributes(IMFAttributes** ppAttributes)
 
 	if (SUCCEEDED(hr))
 	{
-		// Normally defined for Windows 8+ platform, since we also support win7, it will not be defined.
-		const GUID MF_SA_D3D11_AWARE = { 0x206b4fc8, 0xfcf9, 0x4c51, { 0xaf, 0xe3, 0x97, 0x64, 0x36, 0x9e, 0x33, 0xa0 } };
-		Attributes->SetUINT32(MF_SA_D3D11_AWARE, TRUE);
+		// MF_SA_D3D11_AWARE Normally defined for Windows 8+ platform, since we also support win7, it will not be defined.
+		const GUID UE_MF_SA_D3D11_AWARE = { 0x206b4fc8, 0xfcf9, 0x4c51, { 0xaf, 0xe3, 0x97, 0x64, 0x36, 0x9e, 0x33, 0xa0 } };
+		Attributes->SetUINT32(UE_MF_SA_D3D11_AWARE, TRUE);
 		*ppAttributes = Attributes;
 		(*ppAttributes)->AddRef();
 		return S_OK;
