Removing incorrectly set flag indicating support for blocked/synced playback from WMFMedia player

diff --git a/Engine/Plugins/Media/WmfMedia/Source/WmfMedia/Private/Wmf/WmfMediaSession.cpp b/Engine/Plugins/Media/WmfMedia/Source/WmfMedia/Private/Wmf/WmfMediaSession.cpp
--- a/Engine/Plugins/Media/WmfMedia/Source/WmfMedia/Private/Wmf/WmfMediaSession.cpp
+++ b/Engine/Plugins/Media/WmfMedia/Source/WmfMedia/Private/Wmf/WmfMediaSession.cpp
@@ -346,13 +346,6 @@ bool FWmfMediaSession::CanControl(EMediaControl Control) const
 
 	FScopeLock Lock(&CriticalSection);
 
-#if WMFMEDIA_PLAYER_VERSION >= 2
-	if (Control == EMediaControl::BlockOnFetch)
-	{
-		return true;
-	}
-#endif // WMFMEDIA_PLAYER_VERSION >= 2
-
 	if (Control == EMediaControl::Pause)
 	{
 		return ((SessionState == EMediaState::Playing) && (((Capabilities & MFSESSIONCAP_PAUSE) != 0) || UnthinnedRates.Contains(0.0f)));
