Excluding WmfMedia from using player facade provided timing control to fix occasional playback stalls & stutters with this older player

#tests Tests on local machine

diff --git a/Engine/Plugins/Media/WmfMedia/Source/WmfMedia/Private/Player/WmfMediaPlayer.cpp b/Engine/Plugins/Media/WmfMedia/Source/WmfMedia/Private/Player/WmfMediaPlayer.cpp
--- a/Engine/Plugins/Media/WmfMedia/Source/WmfMedia/Private/Player/WmfMediaPlayer.cpp
+++ b/Engine/Plugins/Media/WmfMedia/Source/WmfMedia/Private/Player/WmfMediaPlayer.cpp
@@ -19,6 +19,10 @@
 
 #include "Windows/AllowWindowsPlatformTypes.h"
 
+/*
+*/
+#define WMF_PLAYER_ENABLE_INTERNALCLOCK_ONLY	1
+
 
 /* FWmfVideoPlayer structors
  *****************************************************************************/
@@ -254,6 +258,9 @@ bool FWmfMediaPlayer::GetPlayerFeatureFlag(EFeatureFlag Flag) const
 	{
 		case EFeatureFlag::UsePlaybackTimingV2:
 		case EFeatureFlag::PlayerUsesInternalFlushOnSeek:
+#if WMF_PLAYER_ENABLE_INTERNALCLOCK_ONLY
+		case EFeatureFlag::AlwaysPullNewestVideoFrame:
+#endif
 			return true;
 	}
 	return IMediaPlayer::GetPlayerFeatureFlag(Flag);
diff --git a/Engine/Plugins/Media/WmfMedia/Source/WmfMedia/Private/Player/WmfMediaTracks.cpp b/Engine/Plugins/Media/WmfMedia/Source/WmfMedia/Private/Player/WmfMediaTracks.cpp
--- a/Engine/Plugins/Media/WmfMedia/Source/WmfMedia/Private/Player/WmfMediaTracks.cpp
+++ b/Engine/Plugins/Media/WmfMedia/Source/WmfMedia/Private/Player/WmfMediaTracks.cpp
@@ -403,7 +403,7 @@ void FWmfMediaTracks::SeekStarted(const FTimespan& InTime)
 /* IMediaSamples interface
  *****************************************************************************/
 
-bool FWmfMediaTracks::FetchAudio(TRange<FTimespan> TimeRange, TSharedPtr<IMediaAudioSample, ESPMode::ThreadSafe>& OutSample)
+bool FWmfMediaTracks::FetchAudio(TRange<FMediaTimeStamp> TimeRange, TSharedPtr<IMediaAudioSample, ESPMode::ThreadSafe>& OutSample)
 {
 	TSharedPtr<IMediaAudioSample, ESPMode::ThreadSafe> Sample;
 
@@ -412,9 +412,9 @@ bool FWmfMediaTracks::FetchAudio(TRange<FTimespan> TimeRange, TSharedPtr<IMediaA
 		return false;
 	}
 
-	const FTimespan SampleTime = Sample->GetTime().Time;
+	const FMediaTimeStamp SampleTime = Sample->GetTime();
 
-	if (!TimeRange.Overlaps(TRange<FTimespan>(SampleTime, SampleTime + Sample->GetDuration())))
+	if (!TimeRange.Overlaps(TRange<FMediaTimeStamp>(SampleTime, SampleTime + Sample->GetDuration())))
 	{
 		return false;
 	}
@@ -430,7 +430,7 @@ bool FWmfMediaTracks::FetchAudio(TRange<FTimespan> TimeRange, TSharedPtr<IMediaA
 }
 
 
-bool FWmfMediaTracks::FetchCaption(TRange<FTimespan> TimeRange, TSharedPtr<IMediaOverlaySample, ESPMode::ThreadSafe>& OutSample)
+bool FWmfMediaTracks::FetchCaption(TRange<FMediaTimeStamp> TimeRange, TSharedPtr<IMediaOverlaySample, ESPMode::ThreadSafe>& OutSample)
 {
 	TSharedPtr<IMediaOverlaySample, ESPMode::ThreadSafe> Sample;
 
@@ -439,9 +439,9 @@ bool FWmfMediaTracks::FetchCaption(TRange<FTimespan> TimeRange, TSharedPtr<IMedi
 		return false;
 	}
 
-	const FTimespan SampleTime = Sample->GetTime().Time;
+	const FMediaTimeStamp SampleTime = Sample->GetTime();
 
-	if (!TimeRange.Overlaps(TRange<FTimespan>(SampleTime, SampleTime + Sample->GetDuration())))
+	if (!TimeRange.Overlaps(TRange<FMediaTimeStamp>(SampleTime, SampleTime + Sample->GetDuration())))
 	{
 		return false;
 	}
@@ -457,7 +457,7 @@ bool FWmfMediaTracks::FetchCaption(TRange<FTimespan> TimeRange, TSharedPtr<IMedi
 }
 
 
-bool FWmfMediaTracks::FetchMetadata(TRange<FTimespan> TimeRange, TSharedPtr<IMediaBinarySample, ESPMode::ThreadSafe>& OutSample)
+bool FWmfMediaTracks::FetchMetadata(TRange<FMediaTimeStamp> TimeRange, TSharedPtr<IMediaBinarySample, ESPMode::ThreadSafe>& OutSample)
 {
 	TSharedPtr<IMediaBinarySample, ESPMode::ThreadSafe> Sample;
 
@@ -466,9 +466,9 @@ bool FWmfMediaTracks::FetchMetadata(TRange<FTimespan> TimeRange, TSharedPtr<IMed
 		return false;
 	}
 
-	const FTimespan SampleTime = Sample->GetTime().Time;
+	const FMediaTimeStamp SampleTime = Sample->GetTime();
 
-	if (!TimeRange.Overlaps(TRange<FTimespan>(SampleTime, SampleTime + Sample->GetDuration())))
+	if (!TimeRange.Overlaps(TRange<FMediaTimeStamp>(SampleTime, SampleTime + Sample->GetDuration())))
 	{
 		return false;
 	}
@@ -484,7 +484,7 @@ bool FWmfMediaTracks::FetchMetadata(TRange<FTimespan> TimeRange, TSharedPtr<IMed
 }
 
 
-bool FWmfMediaTracks::FetchVideo(TRange<FTimespan> TimeRange, TSharedPtr<IMediaTextureSample, ESPMode::ThreadSafe>& OutSample)
+bool FWmfMediaTracks::FetchVideo(TRange<FMediaTimeStamp> TimeRange, TSharedPtr<IMediaTextureSample, ESPMode::ThreadSafe>& OutSample)
 {
 	TSharedPtr<IMediaTextureSample, ESPMode::ThreadSafe> Sample;
 
@@ -493,9 +493,9 @@ bool FWmfMediaTracks::FetchVideo(TRange<FTimespan> TimeRange, TSharedPtr<IMediaT
 		return false;
 	}
 
-	const FTimespan SampleTime = Sample->GetTime().Time;
+	const FMediaTimeStamp SampleTime = Sample->GetTime();
 
-	if (!TimeRange.Overlaps(TRange<FTimespan>(SampleTime, SampleTime + Sample->GetDuration())))
+	if (!TimeRange.Overlaps(TRange<FMediaTimeStamp>(SampleTime, SampleTime + Sample->GetDuration())))
 	{
 		return false;
 	}
diff --git a/Engine/Plugins/Media/WmfMedia/Source/WmfMedia/Private/Player/WmfMediaTracks.h b/Engine/Plugins/Media/WmfMedia/Source/WmfMedia/Private/Player/WmfMediaTracks.h
--- a/Engine/Plugins/Media/WmfMedia/Source/WmfMedia/Private/Player/WmfMediaTracks.h
+++ b/Engine/Plugins/Media/WmfMedia/Source/WmfMedia/Private/Player/WmfMediaTracks.h
@@ -201,15 +201,15 @@ public:
 
 	//~ IMediaSamples interface
 
-	virtual bool FetchAudio(TRange<FTimespan> TimeRange, TSharedPtr<IMediaAudioSample, ESPMode::ThreadSafe>& OutSample) override;
-	virtual bool FetchCaption(TRange<FTimespan> TimeRange, TSharedPtr<IMediaOverlaySample, ESPMode::ThreadSafe>& OutSample) override;
-	virtual bool FetchMetadata(TRange<FTimespan> TimeRange, TSharedPtr<IMediaBinarySample, ESPMode::ThreadSafe>& OutSample) override;
-	virtual bool FetchVideo(TRange<FTimespan> TimeRange, TSharedPtr<IMediaTextureSample, ESPMode::ThreadSafe>& OutSample) override;
-
-	virtual bool FetchAudio(TRange<FMediaTimeStamp> TimeRange, TSharedPtr<IMediaAudioSample, ESPMode::ThreadSafe>& OutSample) override { return FetchAudio(TimeRange.GetLowerBound().IsOpen() ? TRange<FTimespan>() : TRange<FTimespan>(TimeRange.GetLowerBoundValue().Time, TimeRange.GetUpperBoundValue().Time), OutSample); }
-	virtual bool FetchCaption(TRange<FMediaTimeStamp> TimeRange, TSharedPtr<IMediaOverlaySample, ESPMode::ThreadSafe>& OutSample) override { return FetchCaption(TRange<FTimespan>(TimeRange.GetLowerBoundValue().Time, TimeRange.GetUpperBoundValue().Time), OutSample); }
-	virtual bool FetchMetadata(TRange<FMediaTimeStamp> TimeRange, TSharedPtr<IMediaBinarySample, ESPMode::ThreadSafe>& OutSample) override { return FetchMetadata(TRange<FTimespan>(TimeRange.GetLowerBoundValue().Time, TimeRange.GetUpperBoundValue().Time), OutSample); }
-	virtual bool FetchVideo(TRange<FMediaTimeStamp> TimeRange, TSharedPtr<IMediaTextureSample, ESPMode::ThreadSafe>& OutSample) override { return FetchVideo(TRange<FTimespan>(TimeRange.GetLowerBoundValue().Time, TimeRange.GetUpperBoundValue().Time), OutSample); }
+	virtual bool FetchAudio(TRange<FTimespan> TimeRange, TSharedPtr<IMediaAudioSample, ESPMode::ThreadSafe>& OutSample) override { return FetchAudio(TimeRange.GetLowerBound().IsOpen() ? TRange<FMediaTimeStamp>() : TRange<FMediaTimeStamp>(FMediaTimeStamp(TimeRange.GetLowerBoundValue(), 0), FMediaTimeStamp(TimeRange.GetUpperBoundValue(), 0)), OutSample); }
+	virtual bool FetchCaption(TRange<FTimespan> TimeRange, TSharedPtr<IMediaOverlaySample, ESPMode::ThreadSafe>& OutSample) override { return FetchCaption(TimeRange.GetLowerBound().IsOpen() ? TRange<FMediaTimeStamp>() : TRange<FMediaTimeStamp>(FMediaTimeStamp(TimeRange.GetLowerBoundValue(), 0), FMediaTimeStamp(TimeRange.GetUpperBoundValue(), 0)), OutSample); }
+	virtual bool FetchMetadata(TRange<FTimespan> TimeRange, TSharedPtr<IMediaBinarySample, ESPMode::ThreadSafe>& OutSample) override { return FetchMetadata(TimeRange.GetLowerBound().IsOpen() ? TRange<FMediaTimeStamp>() : TRange<FMediaTimeStamp>(FMediaTimeStamp(TimeRange.GetLowerBoundValue(), 0), FMediaTimeStamp(TimeRange.GetUpperBoundValue(), 0)), OutSample); }
+	virtual bool FetchVideo(TRange<FTimespan> TimeRange, TSharedPtr<IMediaTextureSample, ESPMode::ThreadSafe>& OutSample) override { return FetchVideo(TimeRange.GetLowerBound().IsOpen() ? TRange<FMediaTimeStamp>() : TRange<FMediaTimeStamp>(FMediaTimeStamp(TimeRange.GetLowerBoundValue(), 0), FMediaTimeStamp(TimeRange.GetUpperBoundValue(), 0)), OutSample); }
+
+	virtual bool FetchAudio(TRange<FMediaTimeStamp> TimeRange, TSharedPtr<IMediaAudioSample, ESPMode::ThreadSafe>& OutSample) override;
+	virtual bool FetchCaption(TRange<FMediaTimeStamp> TimeRange, TSharedPtr<IMediaOverlaySample, ESPMode::ThreadSafe>& OutSample) override;
+	virtual bool FetchMetadata(TRange<FMediaTimeStamp> TimeRange, TSharedPtr<IMediaBinarySample, ESPMode::ThreadSafe>& OutSample) override;
+	virtual bool FetchVideo(TRange<FMediaTimeStamp> TimeRange, TSharedPtr<IMediaTextureSample, ESPMode::ThreadSafe>& OutSample) override;
 
 	virtual void FlushSamples() override;
 #if WMFMEDIA_PLAYER_VERSION >= 2
