Reworking WmfMedia's state handling of WMF to correct playback timing and sample handling shortcomings
#tests Various tests on local machine using testbed apps

diff --git a/Engine/Plugins/Media/WmfMedia/Source/WmfMedia/Private/Player/WmfMediaPlayer.cpp b/Engine/Plugins/Media/WmfMedia/Source/WmfMedia/Private/Player/WmfMediaPlayer.cpp
--- a/Engine/Plugins/Media/WmfMedia/Source/WmfMedia/Private/Player/WmfMediaPlayer.cpp
+++ b/Engine/Plugins/Media/WmfMedia/Source/WmfMedia/Private/Player/WmfMediaPlayer.cpp
@@ -19,11 +19,6 @@
 
 #include "Windows/AllowWindowsPlatformTypes.h"
 
-/*
-*/
-#define WMF_PLAYER_ENABLE_INTERNALCLOCK_ONLY	1
-
-
 /* FWmfVideoPlayer structors
  *****************************************************************************/
 
@@ -35,9 +30,8 @@ FWmfMediaPlayer::FWmfMediaPlayer(IMediaEventSink& InEventSink)
 {
 	check(Session != NULL);
 	check(Tracks.IsValid());
-#if WMFMEDIA_PLAYER_VERSION >= 2
 	Session->SetTracks(Tracks);
-#endif // WMFMEDIA_PLAYER_VERSION >= 2
+	Tracks->SetSession(Session);
 }
 
 
@@ -167,11 +161,7 @@ bool FWmfMediaPlayer::Open(const TSharedRef<FArchive, ESPMode::ThreadSafe>& Arch
 	return InitializePlayer(Archive, OriginalUrl, false, nullptr);
 }
 
-#if WMFMEDIA_PLAYER_VERSION == 1
-void FWmfMediaPlayer::TickFetch(FTimespan /*DeltaTime*/, FTimespan /*Timecode*/)
-#else // WMFMEDIA_PLAYER_VERSION == 1
 void FWmfMediaPlayer::Tick()
-#endif // WMFMEDIA_PLAYER_VERSION == 1
 {
 	bool MediaSourceChanged = false;
 	bool TrackSelectionChanged = false;
@@ -226,9 +216,7 @@ void FWmfMediaPlayer::Tick()
 
 void FWmfMediaPlayer::TickInput(FTimespan DeltaTime, FTimespan /*Timecode*/)
 {
-#if WMFMEDIA_PLAYER_VERSION >= 2
 	Tick();
-#endif // WMFMEDIA_PLAYER_VERSION >= 2
 
 	// forward session events
 	TArray<EMediaEvent> OutEvents;
@@ -240,34 +228,30 @@ void FWmfMediaPlayer::TickInput(FTimespan DeltaTime, FTimespan /*Timecode*/)
 	}
 }
 
-#if WMFMEDIA_PLAYER_VERSION >= 2
 
 bool FWmfMediaPlayer::FlushOnSeekStarted() const
 {
 	return true;
 }
 
+
 bool FWmfMediaPlayer::FlushOnSeekCompleted() const
 {
 	return false;
 }
 
+
 bool FWmfMediaPlayer::GetPlayerFeatureFlag(EFeatureFlag Flag) const
 {
 	switch (Flag)
 	{
 		case EFeatureFlag::UsePlaybackTimingV2:
 		case EFeatureFlag::PlayerUsesInternalFlushOnSeek:
-#if WMF_PLAYER_ENABLE_INTERNALCLOCK_ONLY
-		case EFeatureFlag::AlwaysPullNewestVideoFrame:
-#endif
 			return true;
 	}
 	return IMediaPlayer::GetPlayerFeatureFlag(Flag);
 }
 
-#endif // WMFMEDIA_PLAYER_VERSION >= 2
-
 
 /* FWmfMediaPlayer implementation
  *****************************************************************************/
diff --git a/Engine/Plugins/Media/WmfMedia/Source/WmfMedia/Private/Player/WmfMediaPlayer.h b/Engine/Plugins/Media/WmfMedia/Source/WmfMedia/Private/Player/WmfMediaPlayer.h
--- a/Engine/Plugins/Media/WmfMedia/Source/WmfMedia/Private/Player/WmfMediaPlayer.h
+++ b/Engine/Plugins/Media/WmfMedia/Source/WmfMedia/Private/Player/WmfMediaPlayer.h
@@ -57,17 +57,12 @@ public:
 	virtual bool Open(const FString& Url, const IMediaOptions* Options) override;
 	virtual bool Open(const FString& Url, const IMediaOptions* Options, const FMediaPlayerOptions* PlayerOptions) override;
 	virtual bool Open(const TSharedRef<FArchive, ESPMode::ThreadSafe>& Archive, const FString& OriginalUrl, const IMediaOptions* Options) override;
-#if WMFMEDIA_PLAYER_VERSION == 1
-	virtual void TickFetch(FTimespan DeltaTime, FTimespan Timecode) override;
-#endif // WMFMEDIA_PLAYER_VERSION == 1
 
 	virtual void TickInput(FTimespan DeltaTime, FTimespan Timecode) override;
 	
-#if WMFMEDIA_PLAYER_VERSION >= 2
 	virtual bool FlushOnSeekStarted() const override;
 	virtual bool FlushOnSeekCompleted() const override;
 	virtual bool GetPlayerFeatureFlag(EFeatureFlag flag) const override;
-#endif // WMFMEDIA_PLAYER_VERSION >= 2
 
 protected:
 
@@ -82,10 +77,8 @@ protected:
 	bool InitializePlayer(const TSharedPtr<FArchive, ESPMode::ThreadSafe>& Archive, const FString& Url, bool Precache, const FMediaPlayerOptions* PlayerOptions);
 
 private:
-#if WMFMEDIA_PLAYER_VERSION >= 2
 	/** Tick the player. */
 	void Tick();
-#endif // WMFMEDIA_PLAYER_VERSION >= 2
 
 	/** The duration of the currently loaded media. */
 	FTimespan Duration;
diff --git a/Engine/Plugins/Media/WmfMedia/Source/WmfMedia/Private/Player/WmfMediaTracks.cpp b/Engine/Plugins/Media/WmfMedia/Source/WmfMedia/Private/Player/WmfMediaTracks.cpp
--- a/Engine/Plugins/Media/WmfMedia/Source/WmfMedia/Private/Player/WmfMediaTracks.cpp
+++ b/Engine/Plugins/Media/WmfMedia/Source/WmfMedia/Private/Player/WmfMediaTracks.cpp
@@ -30,6 +30,7 @@
 #include "WmfMediaSettings.h"
 #include "Wmf/WmfMediaSink.h"
 #include "Wmf/WmfMediaStreamSink.h"
+#include "Wmf/WmfMediaSession.h"
 #include "Wmf/WmfMediaTopologyLoader.h"
 #include "Wmf/WmfMediaUtils.h"
 
@@ -55,6 +56,7 @@ FWmfMediaTracks::FWmfMediaTracks()
 	, VideoSamplePool(nullptr)
 	, VideoHardwareVideoDecodingSamplePool(nullptr)
 	, SessionState(EMediaState::Closed)
+	, Session(nullptr)
 {}
 
 
@@ -207,7 +209,6 @@ FTimespan FWmfMediaTracks::GetDuration() const
 
 	UINT64 PresentationDuration = 0;
 	HRESULT Result = PresentationDescriptor->GetUINT64(MF_PD_DURATION, &PresentationDuration);
-#if WMFMEDIA_PLAYER_VERSION >= 2
 	if (SUCCEEDED(Result) == false)
 	{
 		// Live streams like webcam do not have a duration.
@@ -217,7 +218,7 @@ FTimespan FWmfMediaTracks::GetDuration() const
 	
 	// The duration reported here for HAP videos can be larger than they really are be by this amount.
 	PresentationDuration -= 10000;
-#endif
+
 	return FTimespan(PresentationDuration);
 }
 
@@ -346,6 +347,7 @@ void FWmfMediaTracks::Shutdown()
 	{
 		VideoHardwareVideoDecodingSamplePool->Reset();
 	}
+	MediaStreamSink = nullptr;
 
 	SelectedAudioTrack = INDEX_NONE;
 	SelectedCaptionTrack = INDEX_NONE;
@@ -370,17 +372,8 @@ void FWmfMediaTracks::Shutdown()
 	MediaSourceChanged = false;
 	SelectionChanged = false;
 
-#if WMFMEDIA_PLAYER_VERSION >= 2
 	SeekIndex = 0;
-	AudioLoopIndex = 0;
-	VideoLoopIndex = 0;
-	MetaDataLoopIndex = 0;
-	CaptionLoopIndex = 0;
-	LastAudioTime.Reset();
-	LastVideoTime.Reset();
-	LastMetaDataTime.Reset();
-	LastCaptionTime.Reset();
-#endif // WMFMEDIA_PLAYER_VERSION >= 2
+	LoopIndex = 0;
 }
 
 void FWmfMediaTracks::SetSessionState(EMediaState InState)
@@ -389,16 +382,78 @@ void FWmfMediaTracks::SetSessionState(EMediaState InState)
 	SessionState = InState;
 }
 
-#if WMFMEDIA_PLAYER_VERSION >= 2
 
-void FWmfMediaTracks::SeekStarted(const FTimespan& InTime)
+void FWmfMediaTracks::SeekStarted(const FTimespan& InTime, uint32 UserIssuedSeeks, float UnpausedSessionRate)
 {
 	UE_LOG(LogWmfMedia, VeryVerbose, TEXT("FWmfMediaTracks::SeekStarted %f"), InTime.GetTotalSeconds());
+
+	{
+	FScopeLock Lock(&CriticalSection);
+	if (UnpausedSessionRate >= 0.0f)
+	{
+		SeekIndex += UserIssuedSeeks;
+	}
+	else
+	{
+		SeekIndex -= UserIssuedSeeks;
+	}
+	LoopIndex = 0;
+	}
+
+	if (MediaStreamSink)
+	{
+		MediaStreamSink->OnSeekStarted();
+	}
+}
+
+void FWmfMediaTracks::LoopStarted(float SessionRate)
+{
 	FScopeLock Lock(&CriticalSection);
-	++SeekIndex;
+	{
+	if (SessionRate >= 0.0f)
+	{
+		++LoopIndex;
+	}
+	else
+	{
+		--LoopIndex;
+	}
+	}
+
+	if (MediaStreamSink)
+	{
+		// We loop, but then: for the sink it's really just a seek
+		MediaStreamSink->OnSeekStarted();
+	}
+}
+
+
+void FWmfMediaTracks::SessionEnded()
+{
+	if (MediaStreamSink)
+	{
+		MediaStreamSink->OnSessionEnded();
+	}
+}
+
+
+void FWmfMediaTracks::RequestMoreVideoDataFromStreamSink()
+{
+	if (MediaStreamSink)
+	{
+		MediaStreamSink->OnRequestMoreData();
+	}
 }
 
-#endif // WMFMEDIA_PLAYER_VERSION >= 2
+
+bool FWmfMediaTracks::ExecuteOnceMediaStreamSinkHasNoPendingRequests(TFunction<void()>&& ExecuteOnIdle)
+{
+	if (MediaStreamSink)
+	{
+		return MediaStreamSink->ExecuteOnceNoPendingRequests(MoveTemp(ExecuteOnIdle));
+	}
+	return false;
+}
 
 /* IMediaSamples interface
  *****************************************************************************/
@@ -500,6 +555,8 @@ bool FWmfMediaTracks::FetchVideo(TRange<FMediaTimeStamp> TimeRange, TSharedPtr<I
 		return false;
 	}
 
+	Session->RequestMoreVideoData();
+
 	if (!VideoSampleQueue.Dequeue(Sample))
 	{
 		return false;
@@ -513,31 +570,52 @@ bool FWmfMediaTracks::FetchVideo(TRange<FMediaTimeStamp> TimeRange, TSharedPtr<I
 
 void FWmfMediaTracks::FlushSamples()
 {
-	UE_LOG(LogWmfMedia, VeryVerbose, TEXT("FWmfMediaTracks::FlushSamples"));
-	AudioSampleQueue.RequestFlush();
-	CaptionSampleQueue.RequestFlush();
-	MetadataSampleQueue.RequestFlush();
-	VideoSampleQueue.RequestFlush();
+	Session->EnqueueWorkItem(FWmfMediaSession::FWorkItem([this](TComPtr<IMFMediaEvent> MediaEvent, MediaEventType EventType, HRESULT EventStatus, FWmfMediaSession::FBaseWorkItemState* BaseState) -> FWmfMediaSession::EWorkItemResult
+		{
+			{
+			FScopeLock Lock(&CriticalSection);
+
+			UE_LOG(LogWmfMedia, VeryVerbose, TEXT("FWmfMediaTracks::FlushSamples"));
+
+			AudioSampleQueue.RequestFlush();
+			CaptionSampleQueue.RequestFlush();
+			MetadataSampleQueue.RequestFlush();
+			VideoSampleQueue.RequestFlush();
+
+			SeekIndex = 0;
+			LoopIndex = 0;
+			}
+
+			Session->Flush();
+
+			return FWmfMediaSession::EWorkItemResult::Done;
+		}));
 }
 
-#if WMFMEDIA_PLAYER_VERSION >= 2
 
 IMediaSamples::EFetchBestSampleResult FWmfMediaTracks::FetchBestVideoSampleForTimeRange(const TRange<FMediaTimeStamp>& TimeRange, TSharedPtr<IMediaTextureSample, ESPMode::ThreadSafe>& OutSample, bool bReverse, bool bConsistentResult)
 {
-	return VideoSampleQueue.FetchBestSampleForTimeRange(TimeRange, OutSample, bReverse, bConsistentResult) ? EFetchBestSampleResult::Ok : EFetchBestSampleResult::NoSample;
+	EFetchBestSampleResult Ret = VideoSampleQueue.FetchBestSampleForTimeRange(TimeRange, OutSample, bReverse, bConsistentResult) ? EFetchBestSampleResult::Ok : EFetchBestSampleResult::NoSample;
+
+	Session->RequestMoreVideoData();
+
+	return Ret;
 }
 
-#endif // WMFMEDIA_PLAYER_VERSION >= 2
 
 bool FWmfMediaTracks::PeekVideoSampleTime(FMediaTimeStamp & TimeStamp)
 {
 	TSharedPtr<IMediaTextureSample, ESPMode::ThreadSafe> Sample;
-	if (!VideoSampleQueue.Peek(Sample))
+
+	bool Ret = VideoSampleQueue.Peek(Sample);
+
+	Session->RequestMoreVideoData();
+
+	if (Ret)
 	{
-		return false;
+		TimeStamp = Sample->GetTime();
 	}
-	TimeStamp = Sample->GetTime();
-	return true;
+	return Ret;
 }
 
 bool FWmfMediaTracks::DiscardVideoSamples(const TRange<FMediaTimeStamp>& TimeRange, bool bReverse)
@@ -1153,9 +1231,6 @@ bool FWmfMediaTracks::AddTrackToTopology(const FTrack& Track, IMFTopology& Topol
 	// set up output node
 	TComPtr<IMFTopologyNode> OutputNode;
 
-	// Hardware Acccelerated Stream Sink
-	TComPtr<FWmfMediaStreamSink> MediaStreamSink;
-
 	const UWmfMediaSettings* WmfMediaSettings = GetDefault<UWmfMediaSettings>();
 
 	if (VideoSamplePool)
@@ -1192,6 +1267,7 @@ bool FWmfMediaTracks::AddTrackToTopology(const FTrack& Track, IMFTopology& Topol
 	else
 #endif
 	{
+//TODO: >>> SOFTWARE VERSION OR WHAT? THIS IS ALL 100% DIFFERENT?!?!? WTH? (just a "standard" output node takes over) ===========> Sounds like a VERY bad idea!
 		VideoSamplePool = new FWmfMediaTextureSamplePool();
 		if (!VideoSamplePool)
 		{
@@ -1898,28 +1974,7 @@ void FWmfMediaTracks::HandleMediaSamplerClock(EWmfMediaSamplerClockEvent Event,
 FMediaTimeStamp FWmfMediaTracks::AdjustTimeStamp(FTimespan Time, EMediaTrackType TrackType)
 {
 	FScopeLock Lock(&CriticalSection);
-
-	int32* IndexPtr = nullptr;
-	TOptional<FTimespan>* TimePtr = nullptr;
-	switch (TrackType)
-	{
-	case EMediaTrackType::Audio: IndexPtr = &AudioLoopIndex; TimePtr = &LastAudioTime; break;
-	case EMediaTrackType::Video: IndexPtr = &VideoLoopIndex; TimePtr = &LastVideoTime; break;
-	case EMediaTrackType::Metadata: IndexPtr = &MetaDataLoopIndex; TimePtr = &LastMetaDataTime; break;
-	case EMediaTrackType::Caption: IndexPtr = &CaptionLoopIndex; TimePtr = &LastCaptionTime; break;
-	}
-	check(IndexPtr && TimePtr);
-
-	if (TimePtr->IsSet())
-	{
-		if (TimePtr->GetValue() > Time)
-		{
-			++(*IndexPtr);
-		}
-	}
-	*TimePtr = Time;
-
-	return FMediaTimeStamp(Time, FMediaTimeStamp::MakeSequenceIndex(SeekIndex, *IndexPtr));
+	return FMediaTimeStamp(Time, FMediaTimeStamp::MakeSequenceIndex(SeekIndex, LoopIndex));
 }
 
 
diff --git a/Engine/Plugins/Media/WmfMedia/Source/WmfMedia/Private/Player/WmfMediaTracks.h b/Engine/Plugins/Media/WmfMedia/Source/WmfMedia/Private/Player/WmfMediaTracks.h
--- a/Engine/Plugins/Media/WmfMedia/Source/WmfMedia/Private/Player/WmfMediaTracks.h
+++ b/Engine/Plugins/Media/WmfMedia/Source/WmfMedia/Private/Player/WmfMediaTracks.h
@@ -23,12 +23,14 @@ class FWmfMediaAudioSamplePool;
 class FWmfMediaSampler;
 class FWmfMediaTextureSamplePool;
 class FWmfMediaHardwareVideoDecodingTextureSamplePool;
+class FWmfMediaSession;
 class IMediaAudioSample;
 class IMediaBinarySample;
 class IMediaOverlaySample;
 class IMediaTextureSample;
 class FWmfMediaTextureSample;
 class FMediaTimeStamp;
+class FWmfMediaStreamSink;
 
 enum class EMediaTextureSampleFormat;
 enum class EMediaTrackType;
@@ -187,15 +189,28 @@ public:
 	 */
 	void SetSessionState(EMediaState InState);
 
-#if WMFMEDIA_PLAYER_VERSION >= 2
+public:
+
 	/**
-	 * Call this when seeking so that we will discard subsequent samples until
-	 * we find one that matches this seek time.
-	 *
-	 * @param InTime Seek time to wait for.
+	 * Notify track of triggered seek
 	 */
-	void SeekStarted(const FTimespan& InTime);
-#endif // WMFMEDIA_PLAYER_VERSION >= 2
+	void SeekStarted(const FTimespan& InTime, uint32 UserIssuedSeeks, float UnpausedSessionRate);
+
+	/**
+	 * Notify track of triggered loop
+	 */
+	void LoopStarted(float UnpausedSessionRate);
+
+	void SetSession(FWmfMediaSession* InSession)
+	{
+		Session = InSession;
+	}
+
+	void SessionEnded();
+
+	void RequestMoreVideoDataFromStreamSink();
+
+	bool ExecuteOnceMediaStreamSinkHasNoPendingRequests(TFunction<void()>&& ExecuteOnIdle);
 
 public:
 
@@ -212,9 +227,9 @@ public:
 	virtual bool FetchVideo(TRange<FMediaTimeStamp> TimeRange, TSharedPtr<IMediaTextureSample, ESPMode::ThreadSafe>& OutSample) override;
 
 	virtual void FlushSamples() override;
-#if WMFMEDIA_PLAYER_VERSION >= 2
+
 	virtual EFetchBestSampleResult FetchBestVideoSampleForTimeRange(const TRange<FMediaTimeStamp>& TimeRange, TSharedPtr<IMediaTextureSample, ESPMode::ThreadSafe>& OutSample, bool bReverse, bool bConsistentResult) override;
-#endif // WMFMEDIA_PLAYER_VERSION >= 2
+
 	virtual bool DiscardVideoSamples(const TRange<FMediaTimeStamp>& TimeRange, bool bReverse) override;
 	virtual bool DiscardAudioSamples(const TRange<FMediaTimeStamp>& TimeRange, bool bReverse) override;
 	virtual bool DiscardCaptionSamples(const TRange<FMediaTimeStamp>& TimeRange, bool bReverse) override;
@@ -313,7 +328,6 @@ private:
 	FMediaTimeStamp AdjustTimeStamp(FTimespan Time, EMediaTrackType TrackType);
 
 private:
-
 	/** Audio sample object pool. */
 	FWmfMediaAudioSamplePool* AudioSamplePool;
 
@@ -341,6 +355,9 @@ private:
 	/** The currently opened media. */
 	TComPtr<IMFMediaSource> MediaSource;
 
+	// Hardware Acccelerated Stream Sink
+	TComPtr<FWmfMediaStreamSink> MediaStreamSink;
+
 	/** Whether the media source has changed. */
 	bool MediaSourceChanged;
 
@@ -387,18 +404,11 @@ private:
 	/** Current state of the session. */
 	EMediaState SessionState;
 
-#if WMFMEDIA_PLAYER_VERSION >= 2
+	FWmfMediaSession* Session;
+
 	/** SequenceIndex tracking */
 	int32 SeekIndex;
-	int32 AudioLoopIndex;
-	int32 VideoLoopIndex;
-	int32 MetaDataLoopIndex;
-	int32 CaptionLoopIndex;
-	TOptional<FTimespan> LastAudioTime;
-	TOptional<FTimespan> LastVideoTime;
-	TOptional<FTimespan> LastMetaDataTime;
-	TOptional<FTimespan> LastCaptionTime;
-#endif // WMFMEDIA_PLAYER_VERSION >= 2
+	int32 LoopIndex;
 };
 
 
diff --git a/Engine/Plugins/Media/WmfMedia/Source/WmfMedia/Private/Wmf/WmfMediaSession.cpp b/Engine/Plugins/Media/WmfMedia/Source/WmfMedia/Private/Wmf/WmfMediaSession.cpp
--- a/Engine/Plugins/Media/WmfMedia/Source/WmfMedia/Private/Wmf/WmfMediaSession.cpp
+++ b/Engine/Plugins/Media/WmfMedia/Source/WmfMedia/Private/Wmf/WmfMediaSession.cpp
@@ -50,18 +50,19 @@ FWmfMediaSession::FWmfMediaSession()
 	, CurrentDuration(FTimespan::Zero())
 	, MediaSessionCloseEvent(nullptr)
 	, LastTime(FTimespan::Zero())
-	, PendingChanges(false)
 	, RefCount(0)
 	, bIsRequestedTimeLoop(false)
 	, bIsRequestedTimeSeek(false)
+	, LastSetRate(-1.0f)
 	, SessionRate(0.0f)
 	, SessionState(EMediaState::Closed)
 	, ShouldLoop(false)
 	, Status(EMediaStatus::None)
 	, bShuttingDown(false)
-#if WMFMEDIA_PLAYER_VERSION >= 2
 	, bIsWaitingForEnd(false)
-#endif // WMFMEDIA_PLAYER_VERSION >= 2
+	, UserIssuedSeeks(0)
+	, bSeekActive(false)
+	, NumSeeksInWorkQueue(0)
 {
 	UE_LOG(LogWmfMedia, Verbose, TEXT("Session %p: Created"), this);
 	MediaSessionCloseEvent = FPlatformProcess::GetSynchEventFromPool();
@@ -86,28 +87,6 @@ FWmfMediaSession::~FWmfMediaSession()
 
 void FWmfMediaSession::GetEvents(TArray<EMediaEvent>& OutEvents)
 {
-#if WMFMEDIA_PLAYER_VERSION == 1
-#if WMFMEDIASESSION_USE_WINDOWS7FASTFORWARDENDHACK
-	if (CurrentDuration > FTimespan::Zero())
-	{
-		FTimespan Time = GetTime();
-
-		if ((Time < FTimespan::Zero()) || (Time > CurrentDuration))
-		{
-			if (!ShouldLoop)
-			{
-				const HRESULT Result = MediaSession->Stop();
-
-				UE_LOG(LogWmfMedia, Verbose, TEXT("Session %p: Forced media session to stop at end: %s"), this, *WmfMedia::ResultToString(Result));
-			}
-
-			HandleSessionEnded();
-		}
-	}
-#endif
-#endif // WMFMEDIA_PLAYER_VERSION == 1
-
-#if WMFMEDIA_PLAYER_VERSION >= 2
 	// Are we waiting for the end?
 	if (bIsWaitingForEnd)
 	{
@@ -134,8 +113,8 @@ void FWmfMediaSession::GetEvents(TArray<EMediaEvent>& OutEvents)
 			DeferredEvents.Enqueue(EMediaEvent::PlaybackEndReached);
 		}
 	}
-#endif // WMFMEDIA_PLAYER_VERSION >= 2
 
+	// Hand out the events...
 	EMediaEvent Event;
 
 	while (DeferredEvents.Dequeue(Event))
@@ -144,14 +123,12 @@ void FWmfMediaSession::GetEvents(TArray<EMediaEvent>& OutEvents)
 	}
 }
 
-#if WMFMEDIA_PLAYER_VERSION >= 2
 
 void FWmfMediaSession::SetTracks(TSharedPtr<FWmfMediaTracks, ESPMode::ThreadSafe> InTracks)
 {
 	Tracks = InTracks;
 }
 
-#endif // WMFMEDIA_PLAYER_VERSION >= 2
 
 bool FWmfMediaSession::Initialize(bool LowLatency)
 {
@@ -214,60 +191,140 @@ bool FWmfMediaSession::Initialize(bool LowLatency)
 }
 
 
-bool FWmfMediaSession::SetTopology(const TComPtr<IMFTopology>& InTopology, FTimespan InDuration)
+void FWmfMediaSession::EnqueueWorkItem(FWorkItem&& WorkItem)
 {
-	if (MediaSession == NULL)
+	FScopeLock Lock(&CriticalSection);
+
+	// Anything in flight?
+	if (CurrentWorkItem.IsValid())
 	{
-		return false;
+		// Just enqueue the item, we'll see an async notification and then can pump the queue...
+		WorkItemQueue.Push(MoveTemp(WorkItem));
 	}
+	else
+	{
+		// Nothing in flight, so we need to execute in place & react to return codes appropriately...
+		switch (WorkItem.WorkLoad(nullptr, MEUnknown, S_OK, WorkItem.State.Get()))
+		{
+			case EWorkItemResult::Done:			break;
+			case EWorkItemResult::RetryLater:	CurrentWorkItem = WorkItem; break;
+		}
+	}
+}
 
-	UE_LOG(LogWmfMedia, Verbose, TEXT("Session: %p: Setting new partial topology %p (duration = %s)"), this, InTopology.Get(), *InDuration.ToString());
 
+void FWmfMediaSession::JamWorkItem(FWorkItem&& WorkItem)
+{
 	FScopeLock Lock(&CriticalSection);
 
-	if (SessionState == EMediaState::Preparing)
+	// Anything in flight?
+	if (CurrentWorkItem.IsValid())
+	{
+		// Just enqueue the item, we'll see an async notification and then can pump the queue...
+		WorkItemQueue.Insert(MoveTemp(WorkItem), 0);
+	}
+	else
 	{
-		// media source resolved
-		if (InTopology != NULL)
+		// Nothing in flight, so we need to execute in place & react to return codes appropriately...
+		switch (WorkItem.WorkLoad(nullptr, MEUnknown, S_OK, WorkItem.State.Get()))
 		{
-			// at least one track selected
-			HRESULT Result = MediaSession->SetTopology(MFSESSION_SETTOPOLOGY_IMMEDIATE, InTopology);
+			case EWorkItemResult::Done:			break;
+			case EWorkItemResult::RetryLater:	CurrentWorkItem = WorkItem; break;
+		}
+	}
+}
 
-			if (FAILED(Result))
-			{
-				UE_LOG(LogWmfMedia, Verbose, TEXT("Session %p: Failed to set partial topology %p: %s"), this, InTopology.Get(), *WmfMedia::ResultToString(Result));
-				
-				SessionState = EMediaState::Error;
-				DeferredEvents.Enqueue(EMediaEvent::MediaOpenFailed);
-			}
-			else
+
+void FWmfMediaSession::ExecuteNextWorkItems(TComPtr<IMFMediaEvent> MediaEvent, MediaEventType EventType, HRESULT EventStatus)
+{
+	FScopeLock Lock(&CriticalSection);
+
+	do
+	{
+		if (!CurrentWorkItem.IsValid())
+		{
+			if (WorkItemQueue.IsEmpty())
 			{
-				// do nothing (Preparing state will exit in MESessionTopologyStatus event)
+				break;
 			}
+
+			CurrentWorkItem = WorkItemQueue[0];
+			WorkItemQueue.RemoveAt(0);
 		}
-		else
+
+		switch (CurrentWorkItem.WorkLoad(MediaEvent, EventType, EventStatus, CurrentWorkItem.State.Get()))
 		{
-			// no tracks selected
-			UpdateCharacteristics();
-			SessionState = EMediaState::Stopped;
-			DeferredEvents.Enqueue(EMediaEvent::MediaOpened);
+			case EWorkItemResult::Done:			CurrentWorkItem.Reset(); break;	// No new async job, try next workload
+			case EWorkItemResult::RetryLater:	break;							// Workload should be retried later, done for now (implies that some async event or job is ensured to arrive to continue)
 		}
 	}
-	else
+	while (!CurrentWorkItem.IsValid());
+}
+
+
+void FWmfMediaSession::ResetWorkItemQueue()
+{
+	WorkItemQueue.Empty();
+	CurrentWorkItem.Reset();
+	NumSeeksInWorkQueue = 0;
+}
+
+
+bool FWmfMediaSession::SetTopology(const TComPtr<IMFTopology>& InTopology, FTimespan InDuration)
+{
+	if (MediaSession == nullptr)
 	{
-		// topology changed during playback, i.e. track switching
-		if (PendingChanges)
-		{
-			UE_LOG(LogWmfMedia, Verbose, TEXT("Session %p: Requesting topology change after pending command"), this);
-			RequestedTopology = InTopology;
-		}
-		else
-		{
-			CommitTopology(InTopology);
-		}
+		return false;
 	}
 
-	CurrentDuration = InDuration;
+	EnqueueWorkItem(FWorkItem([this, InTopology, InDuration](TComPtr<IMFMediaEvent> MediaEvent, MediaEventType EventType, HRESULT EventStatus, FBaseWorkItemState* BaseState) -> EWorkItemResult
+		{
+			UE_LOG(LogWmfMedia, Verbose, TEXT("Session: %p: Setting new partial topology %p (duration = %s)"), this, InTopology.Get(), *InDuration.ToString());
+
+			EWorkItemResult Res = EWorkItemResult::Done;
+
+			if (SessionState == EMediaState::Preparing)
+			{
+				// media source resolved
+				if (InTopology != nullptr)
+				{
+					// at least one track selected
+					HRESULT Result = MediaSession->SetTopology(MFSESSION_SETTOPOLOGY_IMMEDIATE, InTopology);
+
+					if (FAILED(Result))
+					{
+						UE_LOG(LogWmfMedia, Verbose, TEXT("Session %p: Failed to set partial topology %p: %s"), this, InTopology.Get(), *WmfMedia::ResultToString(Result));
+
+						SessionState = EMediaState::Error;
+						DeferredEvents.Enqueue(EMediaEvent::MediaOpenFailed);
+					}
+					else
+					{
+						// do nothing (Preparing state will exit in MESessionTopologyStatus event)
+						Res = EWorkItemResult::Done;
+					}
+				}
+				else
+				{
+					// no tracks selected
+					UpdateCharacteristics();
+					SessionState = EMediaState::Stopped;
+					DeferredEvents.Enqueue(EMediaEvent::MediaOpened);
+				}
+			}
+			else
+			{
+				// topology changed during playback, i.e. track switching
+				Res = CommitTopology(InTopology, static_cast<FTopologyWorkItemState*>(BaseState), EventType);
+			}
+
+			if (Res != EWorkItemResult::RetryLater)
+			{
+				CurrentDuration = InDuration;
+			}
+
+			return Res;
+		}, MakeShared<FTopologyWorkItemState>()));
 
 	return true;
 }
@@ -285,7 +342,7 @@ void FWmfMediaSession::Shutdown()
 	{
 		// Scope needed since MediaSession->Close() cannot be locked, see below.
 		FScopeLock Lock(&CriticalSection);
-		DiscardPendingChanges();
+		ResetWorkItemQueue();
 		bShuttingDown = true;
 	}
 
@@ -322,7 +379,9 @@ void FWmfMediaSession::Shutdown()
 	CanScrub = false;
 	Capabilities = 0;
 	CurrentDuration = FTimespan::Zero();
+	LastSetRate = -1.0f;
 	SessionRate = 0.0f;
+	UnpausedSessionRate = 0.0f;
 	SessionState = EMediaState::Closed;
 	LastTime = FTimespan::Zero();
 	RequestedRate.Reset();
@@ -330,6 +389,8 @@ void FWmfMediaSession::Shutdown()
 	ThinnedRates.Empty();
 	UnthinnedRates.Empty();
 
+	bSeekActive = false;
+
 	bShuttingDown = false;
 }
 
@@ -372,30 +433,31 @@ bool FWmfMediaSession::CanControl(EMediaControl Control) const
 
 FTimespan FWmfMediaSession::GetDuration() const
 {
+	FScopeLock Lock(&CriticalSection);
 	return CurrentDuration;
 }
 
 
 float FWmfMediaSession::GetRate() const
 {
+	FScopeLock Lock(&CriticalSection);
 	return (SessionState == EMediaState::Playing) ? SessionRate : 0.0f;
 }
 
 
 EMediaState FWmfMediaSession::GetState() const
 {
+	FScopeLock Lock(&CriticalSection);
 	if ((SessionState == EMediaState::Playing) && (SessionRate == 0.0f))
 	{
 		return EMediaState::Paused;
 	}
 
-#if WMFMEDIA_PLAYER_VERSION >= 2
 	// If we are waiting for the end then pretend that we are still playing.
 	if (bIsWaitingForEnd)
 	{
 		return EMediaState::Playing;
 	}
-#endif // WMFMEDIA_PLAYER_VERSION >= 2
 
 	return SessionState;
 }
@@ -403,6 +465,7 @@ EMediaState FWmfMediaSession::GetState() const
 
 EMediaStatus FWmfMediaSession::GetStatus() const
 {
+	FScopeLock Lock(&CriticalSection);
 	return Status;
 }
 
@@ -410,13 +473,16 @@ EMediaStatus FWmfMediaSession::GetStatus() const
 TRangeSet<float> FWmfMediaSession::GetSupportedRates(EMediaRateThinning Thinning) const
 {
 	FScopeLock Lock(&CriticalSection);
-
 	return (Thinning == EMediaRateThinning::Thinned) ? ThinnedRates : UnthinnedRates;
 }
 
 
 FTimespan FWmfMediaSession::GetTime() const
 {
+	/*
+	* Note: for a V2 player this method has not a lot of meaning anymore - The PlayerFacade will generate timing information
+	*/
+
 	FScopeLock Lock(&CriticalSection);
 
 	if (bShuttingDown)
@@ -451,13 +517,14 @@ FTimespan FWmfMediaSession::GetTime() const
 
 bool FWmfMediaSession::IsLooping() const
 {
+	FScopeLock Lock(&CriticalSection);
 	return ShouldLoop;
 }
 
 
 bool FWmfMediaSession::Seek(const FTimespan& Time)
 {
-	if (MediaSession == NULL)
+	if (MediaSession == nullptr)
 	{
 		return false;
 	}
@@ -471,6 +538,12 @@ bool FWmfMediaSession::Seek(const FTimespan& Time)
 		return false;
 	}
 
+	if ((Time < FTimespan::Zero()) || (Time > CurrentDuration))
+	{
+		UE_LOG(LogWmfMedia, Verbose, TEXT("Session %p: Invalid seek time %s (media duration is %s)"), this, *Time.ToString(), *CurrentDuration.ToString());
+		return false;
+	}
+
 	FScopeLock Lock(&CriticalSection);
 
 	if ((SessionState == EMediaState::Closed) || (SessionState == EMediaState::Error))
@@ -479,30 +552,46 @@ bool FWmfMediaSession::Seek(const FTimespan& Time)
 		return false;
 	}
 
-	if ((Time < FTimespan::Zero()) || (Time > CurrentDuration))
-	{
-		UE_LOG(LogWmfMedia, Verbose, TEXT("Session %p: Invalid seek time %s (media duration is %s)"), this, *Time.ToString(), *CurrentDuration.ToString());
-		return false;
-	}
+	++NumSeeksInWorkQueue;
+	EnqueueWorkItem(FWorkItem([this, Time](TComPtr<IMFMediaEvent> MediaEvent, MediaEventType EventType, HRESULT EventStatus, FBaseWorkItemState* BaseState) -> EWorkItemResult
+		{
+			auto State = static_cast<FTimeWorkItemState*>(BaseState);
 
-	// wait for pending changes to complete
-	if (PendingChanges)
-	{
-		UE_LOG(LogWmfMedia, Verbose, TEXT("Session %p: Requesting seek after pending command"), this);
-		RequestedTime = Time;
-		bIsRequestedTimeLoop = false;
-		bIsRequestedTimeSeek = true;
-		return true;
-	}
+			/*
+			* We do not redo or cancel this due to a MESessionEnded being received, as triggering a "start" during this phase will
+			* indeed start the session and interrupt the end of session events (including keeping requests alive etc.)
+			*/
+
+			if (State->State == FTimeWorkItemState::EState::Begin)
+			{
+				++UserIssuedSeeks;
+			}
+			else if (State->State == FTimeWorkItemState::EState::WaitDone)
+			{
+				return (EventType == MESessionStarted) ? EWorkItemResult::Done : EWorkItemResult::RetryLater;
+			}
 
-	return CommitTime(Time, true);
+			check(NumSeeksInWorkQueue > 0);
+			if (--NumSeeksInWorkQueue == 0)
+			{
+				if (CommitTime(Time, true))
+				{
+					State->State = FTimeWorkItemState::EState::WaitDone;
+					return EWorkItemResult::RetryLater;
+				}
+			}
+
+			return EWorkItemResult::Done;
+		}, MakeShared<FTimeWorkItemState, ESPMode::ThreadSafe>()));
+
+	return true;
 }
 
 
 bool FWmfMediaSession::SetLooping(bool Looping)
 {
+	FScopeLock Lock(&CriticalSection);
 	ShouldLoop = Looping;
-
 	return true;
 }
 
@@ -514,23 +603,8 @@ bool FWmfMediaSession::SetRate(float Rate)
 		return false;
 	}
 
-	UE_LOG(LogWmfMedia, Verbose, TEXT("Session %p: Setting rate to %f"), this, Rate);
-
 	FScopeLock Lock(&CriticalSection);
 
-#if 0
-	if (Rate == SessionRate)
-	{
-		// Note: play rates and play states are handled separately in WMF. The session's initial play
-		// rate after opening a media source is 1.0, even though the session is in the Stopped state.
-
-		if ((Rate != 1.0f) || (SessionState == EMediaState::Playing))
-		{
-			return true; // rate already set
-		}
-	}
-#endif
-
 	// validate rate
 	if (!ThinnedRates.Contains(Rate) && !UnthinnedRates.Contains(Rate))
 	{
@@ -538,16 +612,21 @@ bool FWmfMediaSession::SetRate(float Rate)
 		return false;
 	}
 
-	// wait for pending changes to complete
-	if (PendingChanges)
+	if (LastSetRate == Rate)
 	{
-		UE_LOG(LogWmfMedia, Verbose, TEXT("Session %p: Requesting rate change after pending command"), this);
-		RequestedRate = Rate;
-
 		return true;
 	}
 
-	return CommitRate(Rate);
+	LastSetRate = Rate;
+
+	EnqueueWorkItem(FWorkItem([this, Rate](TComPtr<IMFMediaEvent> MediaEvent, MediaEventType EventType, HRESULT EventStatus, FBaseWorkItemState* BaseState) -> EWorkItemResult
+		{
+			UE_LOG(LogWmfMedia, Verbose, TEXT("Session %p: Setting rate to %f"), this, Rate);
+
+			return CommitRate(Rate, static_cast<FRateWorkItemState*>(BaseState), EventType);
+		}, MakeShared<FRateWorkItemState>()));
+
+	return true;
 }
 
 
@@ -570,7 +649,7 @@ STDMETHODIMP FWmfMediaSession::Invoke(IMFAsyncResult* AsyncResult)
 {
 	FScopeLock Lock(&CriticalSection);
 
-	if (MediaSession == NULL)
+	if (MediaSession == nullptr)
 	{
 		return S_OK;
 	}
@@ -637,44 +716,58 @@ STDMETHODIMP FWmfMediaSession::Invoke(IMFAsyncResult* AsyncResult)
 
 	case MESessionCapabilitiesChanged:
 		Capabilities = ::MFGetAttributeUINT32(Event, MF_EVENT_SESSIONCAPS, Capabilities);
+		ExecuteNextWorkItems(Event, EventType, EventStatus);
 		break;
 
 	case MESessionClosed:
 		MediaSessionCloseEvent->Trigger();
 		Capabilities = 0;
 		LastTime = FTimespan::Zero();
+		ResetWorkItemQueue();
+		break;
+
+	case MEEndOfPresentation:
+		HandlePresentationEnded();
+		ExecuteNextWorkItems(Event, EventType, EventStatus);
 		break;
 
 	case MESessionEnded:
 		HandleSessionEnded();
+		ExecuteNextWorkItems(Event, EventType, EventStatus);
 		break;
 
 	case MESessionPaused:
 		HandleSessionPaused(EventStatus);
+		ExecuteNextWorkItems(Event, EventType, EventStatus);
 		break;
 
 	case MESessionRateChanged:
 		HandleSessionRateChanged(EventStatus, *Event);
+		ExecuteNextWorkItems(Event, EventType, EventStatus);
 		break;
 
 	case MESessionScrubSampleComplete:
-		HandleSessionScrubSampleComplete();
+		ExecuteNextWorkItems(Event, EventType, EventStatus);
 		break;
 
 	case MESessionStarted:
 		HandleSessionStarted(EventStatus);
+		ExecuteNextWorkItems(Event, EventType, EventStatus);
 		break;
 
 	case MESessionStopped:
 		HandleSessionStopped(EventStatus);
+		ExecuteNextWorkItems(Event, EventType, EventStatus);
 		break;
 
 	case MESessionTopologySet:
 		HandleSessionTopologySet(EventStatus, *Event);
+		ExecuteNextWorkItems(Event, EventType, EventStatus);
 		break;
 
 	case MESessionTopologyStatus:
 		HandleSessionTopologyStatus(EventStatus, *Event);
+		ExecuteNextWorkItems(Event, EventType, EventStatus);
 		break;
 
 	default:
@@ -695,30 +788,25 @@ STDMETHODIMP FWmfMediaSession::Invoke(IMFAsyncResult* AsyncResult)
 		}
 	}
 
-#if WMFMEDIA_PLAYER_VERSION >= 2
 	// Tell the tracks about our state.
 	TSharedPtr<FWmfMediaTracks, ESPMode::ThreadSafe> TracksPinned = Tracks.Pin();
 	if (TracksPinned.IsValid())
 	{
 		TracksPinned->SetSessionState(GetState());
 	}
-#endif // WMFMEDIA_PLAYER_VERSION >= 2
 
-	UE_LOG(LogWmfMedia, VeryVerbose, TEXT("Session %p: CurrentState: %s, CurrentRate: %f, CurrentTime: %s, SessionState: %s, SessionRate: %f, PendingChanges: %s"),
+	UE_LOG(LogWmfMedia, VeryVerbose, TEXT("Session %p: CurrentState: %s, CurrentRate: %f, CurrentTime: %s, SessionState: %s, SessionRate: %f"),
 		this,
 		WmfMediaSession::StateToString(GetState()),
 		GetRate(),
 		*GetTime().ToString(),
 		WmfMediaSession::StateToString(SessionState),
-		SessionRate,
-		PendingChanges ? TEXT("true") : TEXT("false")
+		SessionRate
 	);
 
-#if WMFMEDIA_PLAYER_VERSION >= 2
 	UE_LOG(LogWmfMedia, VeryVerbose, TEXT("Session %p: WaitingForEnd:%d"),
 		this,
 		bIsWaitingForEnd);
-#endif // WMFMEDIA_PLAYER_VERSION >= 2
 	
 	return S_OK;
 }
@@ -761,215 +849,420 @@ STDMETHODIMP_(ULONG) FWmfMediaSession::Release()
 /* FWmfMediaSession implementation
  *****************************************************************************/
 
-bool FWmfMediaSession::CommitRate(float Rate)
+FWmfMediaSession::EWorkItemResult FWmfMediaSession::CommitRate(float Rate, FRateWorkItemState* State, MediaEventType EventType)
 {
-	check(MediaSession != NULL);
-	check(!PendingChanges);
+	check(MediaSession != nullptr);
 
-	UE_LOG(LogWmfMedia, Verbose, TEXT("Session %p: Committing rate %f"), this, Rate);
+	UE_LOG(LogWmfMedia, VeryVerbose, TEXT("Session %p: Committing rate %f [S=%d]"), this, Rate, int(State->State));
 
-	// Note: if rate control is not available, the session only supports pause and play
+	// Did session end and we do not just start up?
+	if (EventType == MESessionEnded && State->State != FRateWorkItemState::EState::Begin)
+	{
+		// Make us start over again... what we did so far has been lost in the deep belly of WMF's machinery...
+		State->State = FRateWorkItemState::EState::Begin;
+	}
 
-	if (RateControl == NULL)
+	if (RateControl == nullptr)
 	{
-		if (Rate == 0.0f)
+		/*
+		* No rate control
+		* 
+		* Session is only able to support pause & play
+		* (no reverse, scrubbing or speed changes)
+		*/
+
+		if (State->State == FRateWorkItemState::EState::Begin)
 		{
-			if (SessionState == EMediaState::Playing)
+			if (Rate == 0.0f)
 			{
-				const HRESULT Result = MediaSession->Pause();
-
-				if (FAILED(Result))
+				if (SessionState == EMediaState::Playing)
 				{
-					UE_LOG(LogWmfMedia, Verbose, TEXT("Session %p: Failed to pause session: %s"), this, *WmfMedia::ResultToString(Result));
-					return false;
+					const HRESULT Result = MediaSession->Pause();
+
+					if (FAILED(Result))
+					{
+						UE_LOG(LogWmfMedia, Verbose, TEXT("Session %p: Failed to pause session: %s"), this, *WmfMedia::ResultToString(Result));
+						return EWorkItemResult::Done;
+					}
+
+					State->State = FRateWorkItemState::EState::WaitForPause;
+					return EWorkItemResult::RetryLater;
 				}
 			}
-		}
-		else
-		{
-			if (SessionState != EMediaState::Playing)
+			else
 			{
-				PROPVARIANT StartPosition;
-				PropVariantInit(&StartPosition);
+				if (SessionState != EMediaState::Playing)
+				{
+					PROPVARIANT StartPosition;
+					PropVariantInit(&StartPosition);
 
-				const HRESULT Result = MediaSession->Start(NULL, &StartPosition);
+					const HRESULT Result = MediaSession->Start(NULL, &StartPosition);
 
-				if (FAILED(Result))
-				{
-					UE_LOG(LogWmfMedia, Verbose, TEXT("Session %p: Failed to start session: %s"), this, *WmfMedia::ResultToString(Result));
-					return false;
+					if (FAILED(Result))
+					{
+						UE_LOG(LogWmfMedia, Verbose, TEXT("Session %p: Failed to start session: %s"), this, *WmfMedia::ResultToString(Result));
+						return EWorkItemResult::Done;
+					}
+
+					State->State = FRateWorkItemState::EState::WaitForStart;
+					return EWorkItemResult::RetryLater;
 				}
 			}
 		}
+		else if (State->State == FRateWorkItemState::EState::WaitForPause)
+		{
+			if (EventType != MESessionPaused)
+			{
+				return EWorkItemResult::RetryLater;
+			}
+		}
+		else
+		{
+			check(State->State == FRateWorkItemState::EState::WaitForStart);
 
-		PendingChanges = true;
+			if (EventType != MESessionStarted)
+			{
+				return EWorkItemResult::RetryLater;
+			}
+		}
 
-		return true;
+		return EWorkItemResult::Done;
 	}
 
-	// Note: if rate control is available, things get considerably more complicated
-	// as many rate transitions are only allowed from certain session states
+	// ---------------------------------------------------------------------------------------------------------------
+	/**
+	 * Rate control is present
+	 * 
+	 * Session can plat at variable speeds, reverse (possibly) and scrub (possibly)
+	 */
 
-	if (((Rate >= 0.0f) && (SessionRate < 0.0f)) || ((Rate < 0.0f) && (SessionRate >= 0.0f)))
+	if (State->State == FRateWorkItemState::EState::Begin)
 	{
-		// transitions between negative and zero/positive rates require Stopped state
-		if (SessionState != EMediaState::Stopped)
+		State->State = FRateWorkItemState::EState::Ready;
+		State->LastTime = WmfMediaSession::RequestedTimeCurrent;
+
+		if (Rate != 0.0f && ((UnpausedSessionRate * Rate) < 0.0f))
 		{
-			UE_LOG(LogWmfMedia, Verbose, TEXT("Session %p: Stopping session for rate change"), this);
+			/**
+			 * Reversal of playback direction
+			 * 
+			 * System needs to be stopped and we ensure no pending requests as WMF tends to get confused if they come in while we do this
+			 */
+
+			// Stopped already?
+			if (SessionState != EMediaState::Stopped)
+			{
+				// No...
+				
+				// If we reverse, we might loose issued sample requests (without us knowing) - wait until all current ones are done
+				TSharedPtr<FWmfMediaTracks, ESPMode::ThreadSafe> TracksPinned = Tracks.Pin();
+				if (TracksPinned.IsValid())
+				{
+					// Schedule the lambda to be executed once all requests are done if any are pending...
+					if (TracksPinned->ExecuteOnceMediaStreamSinkHasNoPendingRequests([this]()
+						{
+							// If we have pending sample requests, this will be executed by the sink once it detects all of them being done.
+							// All we do is to pump our workitem queue once more to ensure that - no matter what internal async wakeups we get we at least get this one to continue!
+							FScopeLock Lock(&CriticalSection);
+							ExecuteNextWorkItems(nullptr, MEUnknown, S_OK);
+						}))
+					{
+						// We had pending requests. Reset to original state and a later retry
+						State->State = FRateWorkItemState::EState::Begin;
+						return EWorkItemResult::RetryLater;
+					}
+				}
 
-			LastTime = GetTime();
+				// Gather current clock, so we can restart at the same spot after the stop/reverse step
+				LONGLONG ClockTime;
+				MFTIME SystemTime;
+				if (PresentationClock->GetCorrelatedTime(0, &ClockTime, &SystemTime) == S_OK)
+				{
+					State->LastTime = FTimespan(ClockTime);
+				}
 
-			// stop the media session
-			const HRESULT Result = MediaSession->Stop();
+				// Stop the media session
+				const HRESULT Result = MediaSession->Stop();
+				if (FAILED(Result))
+				{
+					UE_LOG(LogWmfMedia, Verbose, TEXT("Session %p: Failed to stop for rate change: %s"), this, *WmfMedia::ResultToString(Result));
+					return EWorkItemResult::Done;
+				}
 
-			if (FAILED(Result))
-			{
-				UE_LOG(LogWmfMedia, Verbose, TEXT("Session %p: Failed to stop for rate change: %s"), this, *WmfMedia::ResultToString(Result));
-				return false;
+				// Next wait for this to finish and then set & restart...
+				State->State = FRateWorkItemState::EState::WaitForStopThenSetAndRestart;
+				return EWorkItemResult::RetryLater;
 			}
 
-			// defer rate change
-			UE_LOG(LogWmfMedia, Verbose, TEXT("Session %p: Deferring rate change until after pending stop"), this);
+			// We are stopped. Just continue to the rate set...
+		}
+		else
+		{
+			if (Rate != 0.0f)
+			{
+				if (SessionState != EMediaState::Playing)
+				{
 
-			if (!RequestedTime.IsSet())
-			{			
-				RequestedTime = LastTime;
-				bIsRequestedTimeLoop = false;
-				bIsRequestedTimeSeek = false;
+					// Going out of "paused" state, but we want to actually we have been "scrubbing", which WMF interprets as a PLAYING state!
+					// So we need to pause first, to then restart normal playback (or else the rate change fails)
+					// [this also triggers on initial startup, but it seems harmless to pause for no reason]
+
+					if (SessionState != EMediaState::Stopped)
+					{
+						/**
+						 * Leave paused state
+						 */
+
+						if (CanScrub)
+						{
+							/*
+							* As we can scrub, we actually are leaving "scrub state", not a real pause, but a playback state as far as WMF is concerned!
+							*/
+
+							// Was the last non-scrubbing playback reverse?
+							if (UnpausedSessionRate >= 0.0f)
+							{
+								// No. We need to pause to be able to set a new non-zero rate...
+								const HRESULT Result = MediaSession->Pause();
+								if (FAILED(Result))
+								{
+									UE_LOG(LogWmfMedia, Verbose, TEXT("Session %p: Failed to pause for rate change: %s"), this, *WmfMedia::ResultToString(Result));
+									return EWorkItemResult::Done;
+								}
+
+								// Wait for pause to be done, then set and restart...
+								State->State = FRateWorkItemState::EState::WaitForPauseThenSetAndRestart;
+							}
+							else
+							{
+								// Yes. We need to actually stop session as we are going to go back into a reverse playback situation from scrubbing
+
+								// Get the time so we can restart at the correct spot
+								LONGLONG ClockTime;
+								MFTIME SystemTime;
+								if (PresentationClock->GetCorrelatedTime(0, &ClockTime, &SystemTime) == S_OK)
+								{
+									State->LastTime = FTimespan(ClockTime);
+								}
+								
+								// Stop the session
+								const HRESULT Result = MediaSession->Stop();
+								if (FAILED(Result))
+								{
+									UE_LOG(LogWmfMedia, Verbose, TEXT("Session %p: Failed to stop for rate change: %s"), this, *WmfMedia::ResultToString(Result));
+									return EWorkItemResult::Done;
+								}
+
+								// Wait for pause to be done, then set and restart...
+								State->State = FRateWorkItemState::EState::WaitForStopThenSetAndRestart;
+							}
+						}
+						else
+						{
+							/**
+							 * No scrubbing. We paused "for real". No just start things up again...
+							 */
+							PROPVARIANT StartPosition;
+							PropVariantInit(&StartPosition);
+
+							const HRESULT Result = MediaSession->Start(NULL, &StartPosition);
+
+							if (FAILED(Result))
+							{
+								UE_LOG(LogWmfMedia, Verbose, TEXT("Session %p: Failed to start session: %s"), this, *WmfMedia::ResultToString(Result));
+								return EWorkItemResult::Done;
+							}
+
+							State->State = FRateWorkItemState::EState::WaitForStart;
+						}
+
+						return EWorkItemResult::RetryLater;
+					}
+
+					// We are stopped. We can set a new rate and restart without any other work before that...
+					State->State = FRateWorkItemState::EState::ReadyNeedsRestart;
+				}
 			}
+			else
+			{
+				if (SessionState == EMediaState::Playing)
+				{
+					/**
+					 * Entering pause state
+					 */
+
+					if (CanScrub)
+					{
+						/*
+						* We can scrub. Hence we rather enter 'scrub mode' than a real 'pause' state
+						*/
+
+						// Wait for any pending requests as WMF tends to loose track of these during the transition...
+						TSharedPtr<FWmfMediaTracks, ESPMode::ThreadSafe> TracksPinned = Tracks.Pin();
+						if (TracksPinned.IsValid())
+						{
+							// Schedule the lambda to be run once we have no outstanding requests anymore if we ever had any
+							if (TracksPinned->ExecuteOnceMediaStreamSinkHasNoPendingRequests([this]()
+								{
+									// If we have pending sample requests, this will be executed by the sink once it detects all of them being done.
+									// All we do is to pump our workitem queue once more to ensure that - no matter what internal async wakeups we get we at least get this one to continue!
+									FScopeLock Lock(&CriticalSection);
+									ExecuteNextWorkItems(nullptr, MEUnknown, S_OK);
+								}))
+							{
+								// We had outstanding requests, so we reset state and retry later
+								State->State = FRateWorkItemState::EState::Begin;
+								return EWorkItemResult::RetryLater;
+							}
+						}
+					}
+
+					// Last unpaused rate was forward OR we cannot scrub...
+					if (UnpausedSessionRate >= 0.0f || !CanScrub)
+					{
+						// Pause the session
+						const HRESULT Result = MediaSession->Pause();
+						if (FAILED(Result))
+						{
+							UE_LOG(LogWmfMedia, Verbose, TEXT("Session %p: Failed to pause for rate change: %s"), this, *WmfMedia::ResultToString(Result));
+							return EWorkItemResult::Done;
+						}
+
+						// Continue to start up into scrub mode (start with rate == 0) or simply leave things plaused...
+						State->State = CanScrub ? FRateWorkItemState::EState::WaitForPauseThenSetAndRestart : FRateWorkItemState::EState::WaitForPause;
+					}
+					else
+					{
+						// We had been in reverse. We need to stop the session to be able to go into 'srub mode'...
+
+						// Get the current playback time so we can restart at the correct spot
+						LONGLONG ClockTime;
+						MFTIME SystemTime;
+						if (PresentationClock->GetCorrelatedTime(0, &ClockTime, &SystemTime) == S_OK)
+						{
+							State->LastTime = FTimespan(ClockTime);
+						}
+ 
+						// Stop the session
+						const HRESULT Result = MediaSession->Stop();
+						if (FAILED(Result))
+						{
+							UE_LOG(LogWmfMedia, Verbose, TEXT("Session %p: Failed to stop for rate change: %s"), this, *WmfMedia::ResultToString(Result));
+							return EWorkItemResult::Done;
+						}
+
+						// Wait for the stop and then continue with a restart into 'scrub mode'
+						State->State = FRateWorkItemState::EState::WaitForStopThenSetAndRestart;
+					}
+
+					return EWorkItemResult::RetryLater;
+				}
 
-			RequestedRate = Rate;
-			PendingChanges = true;
-
-			return true;
+				// Just set the rate now. No restart...
+			}
 		}
 	}
-	
-	if (((Rate == 0.0f) && (SessionRate != 0.0f)) || ((Rate != 0.0f) && (SessionRate == 0.0f)))
+	else if (State->State == FRateWorkItemState::EState::WaitForStopThenSetAndRestart)
 	{
-		// transitions between positive and zero rates require Paused or Stopped state
-		if ((SessionState != EMediaState::Paused) && (SessionState != EMediaState::Stopped))
+		// Wait for stop and then set and restart...
+		if (EventType != MESessionStopped)
 		{
-			UE_LOG(LogWmfMedia, Verbose, TEXT("Session %p: Pausing session for rate change from %f to %f"), this, SessionRate, Rate);
-
-			// pause the media session
-			const HRESULT Result = MediaSession->Pause();
-
-			if (FAILED(Result))
-			{
-				UE_LOG(LogWmfMedia, Verbose, TEXT("Session %p: Failed to pause for rate change: %s"), this, *WmfMedia::ResultToString(Result));
-				return false;
-			}
-
-			// defer rate change
-			UE_LOG(LogWmfMedia, Verbose, TEXT("Session %p: Deferring rate change until after pending pause"), this);
-
-			RequestedRate = Rate;
-			PendingChanges = true;
-
-			return true;
+			return EWorkItemResult::RetryLater;
 		}
+		State->State = FRateWorkItemState::EState::ReadyNeedsRestart;
 	}
-
-	// Note: The rate control could be updated right after requesting the Pause or Stopped
-	// states above, but we wait for these transitions to complete, so that multiple calls
-	// to SetRate do not interfere with each other.
-
-	if (Rate != SessionRate)
+	else if (State->State == FRateWorkItemState::EState::WaitForPauseThenSetAndRestart)
 	{
-		// determine thinning mode
-		EMediaRateThinning Thinning;
-
-		if (UnthinnedRates.Contains(Rate))
+		// Wait for pause and then set and restart...
+		if (EventType != MESessionPaused)
 		{
-			Thinning = EMediaRateThinning::Unthinned;
+			return EWorkItemResult::RetryLater;
 		}
-		else if (ThinnedRates.Contains(Rate))
+		State->State = FRateWorkItemState::EState::ReadyNeedsRestart;
+	}
+	else if (State->State == FRateWorkItemState::EState::WaitForPause)
+	{
+		// Check if pause is all done!
+		return (EventType == MESessionPaused) ? EWorkItemResult::Done : EWorkItemResult::RetryLater;
+	}
+	else if (State->State == FRateWorkItemState::EState::WaitForSetAndRestart)
+	{
+		// Wait for rate set and then restart...
+		if (EventType != MESessionRateChanged)
 		{
-			Thinning = EMediaRateThinning::Thinned;
+			return EWorkItemResult::RetryLater;
 		}
-		else
+
+		// Restart the session at the given time (if none is specified, we assume the 'current time' is valid with the session)
+		PROPVARIANT StartPosition;
+		PropVariantInit(&StartPosition);
+		if (State->LastTime != WmfMediaSession::RequestedTimeCurrent)
 		{
-			return false;
+			StartPosition.vt = VT_I8;
+			StartPosition.hVal.QuadPart = State->LastTime.GetTicks();
 		}
 
-		const TCHAR* ThinnedString = (Thinning == EMediaRateThinning::Thinned) ? TEXT("thinned") : TEXT("unthinned");
-		UE_LOG(LogWmfMedia, Verbose, TEXT("Session %p: Changing rate from %f to %f [%s]"), this, SessionRate, Rate, ThinnedString);
-
-		// set the new rate
-		const HRESULT Result = RateControl->SetRate((Thinning == EMediaRateThinning::Thinned) ? TRUE : FALSE, Rate);
-
+		const HRESULT Result = MediaSession->Start(NULL, &StartPosition);
 		if (FAILED(Result))
 		{
-			UE_LOG(LogWmfMedia, Verbose, TEXT("Session %p: Failed to change rate: %s"), this, *WmfMedia::ResultToString(Result));
-			return false;
+			UE_LOG(LogWmfMedia, Verbose, TEXT("Session %p: Failed to start session: %s"), this, *WmfMedia::ResultToString(Result));
+			return EWorkItemResult::Done;
 		}
 
-		if (PendingChanges)
+		// Now wait for the restart to be done...
+		State->State = FRateWorkItemState::EState::WaitForRestart;
+		return EWorkItemResult::RetryLater;
+	}
+	else if (State->State == FRateWorkItemState::EState::WaitForSet)
+	{
+		// Wait for rate being set (no restart)
+		if (EventType != MESessionRateChanged)
 		{
-			return true; // wait for required state transitions to complete
+			return EWorkItemResult::RetryLater;
 		}
-
-		PendingChanges = true;
+		return EWorkItemResult::Done;
 	}
-
-	// Note: no further changes needed if session was playing and direction didn't change
-
-	if (((Rate * SessionRate) > 0.0f) && (SessionState == EMediaState::Playing))
+	else if (State->State == FRateWorkItemState::EState::WaitForRestart)
 	{
-		return true;
+		// Check if restart is all done!
+		return (EventType == MESessionStarted) ? EWorkItemResult::Done : EWorkItemResult::RetryLater;
 	}
 
-	// Note: for non-zero rates, the session must be restarted. If the rate control wasn't
-	// updated above, this can be done immediately, otherwise it has to be deferred until
-	// after the rate control finished setting the new rate.
-
-	if ((Rate != 0.0f) && ((SessionState != EMediaState::Playing) || PendingChanges))
+	// determine thinning mode
+	EMediaRateThinning Thinning;
+	if (UnthinnedRates.Contains(Rate))
 	{
-		// determine restart time
-		FTimespan RestartTime;
-
-		if (RequestedTime.IsSet())
-		{
-			RestartTime = RequestedTime.GetValue();
-		}
-		else
-		{
-			RestartTime = GetTime();
-		}
-
-		if ((RestartTime == FTimespan::Zero()) && (Rate < 0.0f))
-		{
-			RestartTime = CurrentDuration; // loop to end
-		}
-		else if ((RestartTime == CurrentDuration) && (Rate > 0.0f))
-		{
-			RestartTime = FTimespan::Zero(); // loop to beginning
-		}
+		Thinning = EMediaRateThinning::Unthinned;
+	}
+	else if (ThinnedRates.Contains(Rate))
+	{
+		Thinning = EMediaRateThinning::Thinned;
+	}
+	else
+	{
+		return EWorkItemResult::Done;
+	}
 
-		// restart session
-		if (PendingChanges)
-		{
-			UE_LOG(LogWmfMedia, Verbose, TEXT("Session %p: Requesting start after pending rate change"), this);
-			RequestedTime = RestartTime;
-			bIsRequestedTimeLoop = false;
-			bIsRequestedTimeSeek = false;
-		}
-		else
+	// ...and set the rate
+	{
+		const HRESULT Result = RateControl->SetRate((Thinning == EMediaRateThinning::Thinned) ? TRUE : FALSE, Rate);
+		if (FAILED(Result))
 		{
-			UE_LOG(LogWmfMedia, Verbose, TEXT("Session %p: Starting session for rate change"), this);
-			CommitTime(RestartTime, false);
+			UE_LOG(LogWmfMedia, Verbose, TEXT("Session %p: Failed to change rate: %s"), this, *WmfMedia::ResultToString(Result));
+			return EWorkItemResult::Done;
 		}
 	}
 
-	return true;
+	// Now we wait for the rate to be set and possibly trigger a restart...
+	State->State = (State->State == FRateWorkItemState::EState::ReadyNeedsRestart) ? FRateWorkItemState::EState::WaitForSetAndRestart : FRateWorkItemState::EState::WaitForSet;
+
+	return EWorkItemResult::RetryLater;
 }
 
 
 bool FWmfMediaSession::CommitTime(FTimespan Time, bool bIsSeek)
 {
 	check(MediaSession != NULL);
-	check(PendingChanges == false);
 
 	FTimespan OriginalTime = Time;
 	const FString TimeString = (Time == WmfMediaSession::RequestedTimeCurrent) ? TEXT("<current>") : *Time.ToString();
@@ -983,19 +1276,6 @@ bool FWmfMediaSession::CommitTime(FTimespan Time, bool bIsSeek)
 		UE_LOG(LogWmfMedia, Verbose, TEXT("Session %p: Starting from <current>, because media can't seek"), this);
 		Time = WmfMediaSession::RequestedTimeCurrent;
 	}
-#if WMFMEDIA_PLAYER_VERSION == 1
-	else if (Time == GetTime())
-#else // WMFMEDIA_PLAYER_VERSION == 1
-	// Only do this whan the rate is non zero, otherwise we will not ask for the correct time.
-	else if ((Time == GetTime()) && (SessionRate != 0.0f))
-#endif // WMFMEDIA_PLAYER_VERSION == 1
-	{
-		// Fix for audio desync and video fast-forwarding behavior
-		// There long delay (500ms+) until samples start arriving unless we specifically use RequestedTimeCurrent
-		// After delay occurs samples begin arriving at accelerated speed until caught up to playback time leading to visual and audio problems
-		UE_LOG(LogWmfMedia, Verbose, TEXT("Session %p: Starting from <current>, because media already queued up to correct time"), this);
-		Time = WmfMediaSession::RequestedTimeCurrent;
-	}
 
 	if (Time == WmfMediaSession::RequestedTimeCurrent)
 	{
@@ -1015,165 +1295,125 @@ bool FWmfMediaSession::CommitTime(FTimespan Time, bool bIsSeek)
 		return false;
 	}
 
-#if WMFMEDIA_PLAYER_VERSION >= 2
 	// If this is not a loop, then tell the tracks about the seek.
 	if (bCanSeek && bIsSeek)
 	{
 		TSharedPtr<FWmfMediaTracks, ESPMode::ThreadSafe> TracksPinned = Tracks.Pin();
 		if (TracksPinned.IsValid())
 		{
-			TracksPinned->SeekStarted(OriginalTime);
+			TracksPinned->SeekStarted(OriginalTime, UserIssuedSeeks, UnpausedSessionRate);
+			UserIssuedSeeks = 0;
 		}
+		bSeekActive = true;
 	}
-#endif // WMFMEDIA_PLAYER_VERSION >= 2
-
-	PendingChanges = true;
 
 	return true;
 }
 
 
-bool FWmfMediaSession::CommitTopology(IMFTopology* Topology)
+FWmfMediaSession::EWorkItemResult FWmfMediaSession::CommitTopology(IMFTopology* Topology, FTopologyWorkItemState* State, MediaEventType EventType)
 {
 	check(MediaSession != NULL);
-	check(PendingChanges == false);
 
-	UE_LOG(LogWmfMedia, Verbose, TEXT("Session %p: Committing topology %p"), this, Topology);
+	UE_LOG(LogWmfMedia, Verbose, TEXT("Session %p: Committing topology %p [S=%d]"), this, Topology, int(State->State));
 
-	if (SessionState != EMediaState::Stopped)
+	// Did session end and we do not just start up?
+	if (EventType == MESessionEnded && State->State != FTopologyWorkItemState::EState::Begin)
 	{
-		UE_LOG(LogWmfMedia, Verbose, TEXT("Session %p: Stopping session for topology change"), this);
-
-		LastTime = GetTime();
-
-		// topology change requires transition to Stopped; playback is resumed afterwards
-		HRESULT Result = MediaSession->Stop();
+		// Make us start over again... what we did so far has been lost in the deep belly of WMF's machinery...
+		State->State = FTopologyWorkItemState::EState::Begin;
+	}
 
-		if (FAILED(Result))
+	if (State->State == FTopologyWorkItemState::EState::Begin)
+	{
+		// Starting up for the first time. Check if we are stopped already?
+		if (SessionState != EMediaState::Stopped)
 		{
-			UE_LOG(LogWmfMedia, Verbose, TEXT("Session %p: Failed to stop for topology change: %s"), this, *WmfMedia::ResultToString(Result));
-			return false;
-		}
+			// No. Stop session and retry...
+			UE_LOG(LogWmfMedia, Verbose, TEXT("Session %p: Stopping session for topology change"), this);
 
-		// request deferred playback restart
-		if (!RequestedTime.IsSet())
-		{
-			UE_LOG(LogWmfMedia, Verbose, TEXT("Session %p: Requesting restart after pending stop"), this);
+			// topology change requires transition to Stopped; playback is resumed afterwards
+			HRESULT Result = MediaSession->Stop();
+			if (FAILED(Result))
+			{
+				UE_LOG(LogWmfMedia, Verbose, TEXT("Session %p: Failed to stop for topology change: %s"), this, *WmfMedia::ResultToString(Result));
+				return EWorkItemResult::Done;
+			}
 
-			RequestedTime = LastTime;
-			bIsRequestedTimeLoop = false;
-			bIsRequestedTimeSeek = false;
+			State->State = FTopologyWorkItemState::EState::ReadyNeedsRestart;
 
-			// Zero LastTime so it matches what WMF is actually doing (rewinding to the beginning)
-			// The PresentationClock has now stopped and FWmfMediaSession::GetTime() will just return LastTime
-			LastTime = FTimespan::Zero();
+			// Come back later once the session has stopped...
+			return EWorkItemResult::RetryLater;
+		}
+		else
+		{
+			// Proceed to actual setup...
+			State->State = FTopologyWorkItemState::EState::Ready;
 		}
-
-		UE_LOG(LogWmfMedia, Verbose, TEXT("Session %p: Requesting topology change after pending stop"), this);
-
-		RequestedTopology = Topology;
-		PendingChanges = true;
-
-		return true;
 	}
-
-	// set new topology
-	HRESULT Result = MediaSession->ClearTopologies();
-
-	if (FAILED(Result))
+	else if (State->State == FTopologyWorkItemState::EState::Ready)
 	{
-		UE_LOG(LogWmfMedia, Verbose, TEXT("Session %p: Failed to clear queued topologies: %s"), this, *WmfMedia::ResultToString(Result));
-		return false;
+		return (EventType != MESessionTopologySet) ? EWorkItemResult::RetryLater : EWorkItemResult::Done;
 	}
-
-	if (Topology != NULL)
+	else if (State->State == FTopologyWorkItemState::EState::ReadyNeedsRestart)
 	{
-		Result = MediaSession->SetTopology(MFSESSION_SETTOPOLOGY_IMMEDIATE, Topology);
+		if (EventType != MESessionTopologySet)
+		{
+			return EWorkItemResult::RetryLater;
+		}
 
+		// Topology is set, we need to restart the session...
+		PROPVARIANT StartPosition;
+		PropVariantInit(&StartPosition);
+		HRESULT Result = MediaSession->Start(NULL, &StartPosition);
 		if (FAILED(Result))
 		{
-			UE_LOG(LogWmfMedia, Verbose, TEXT("Session %p: Failed to set topology %p: %s"), Topology, this, *WmfMedia::ResultToString(Result));
-			return false;
+			UE_LOG(LogWmfMedia, Verbose, TEXT("Session %p: Failed to start after topology change: %s"), this, *WmfMedia::ResultToString(Result));
+			return EWorkItemResult::Done;
 		}
 
-		UE_LOG(LogWmfMedia, Verbose, TEXT("Session %p: Committed topology %p"), this, Topology);
-
-		PendingChanges = true;
+		State->State = FTopologyWorkItemState::EState::WaitRestart;
+		return EWorkItemResult::RetryLater;
 	}
-
-	return true;
-}
-
-
-void FWmfMediaSession::DiscardPendingChanges()
-{
-	UE_LOG(LogWmfMedia, Verbose, TEXT("Session %p: Discarding pending changes"), this);
-
-	RequestedRate.Reset();
-	RequestedTime.Reset();
-	bIsRequestedTimeLoop = false;
-	bIsRequestedTimeSeek = false;
-	RequestedTopology.Reset();
-
-	PendingChanges = false;
-}
-
-
-void FWmfMediaSession::DoPendingChanges()
-{
-	if (PendingChanges)
+	else if (State->State == FTopologyWorkItemState::EState::WaitRestart)
 	{
-		const FString RequestedRateString = !RequestedRate.IsSet() ? TEXT("<none>") : FString::Printf(TEXT("%f"), RequestedRate.GetValue());
-		const FString RequestedTimeString = !RequestedTime.IsSet() ? TEXT("<none>") : (RequestedTime.GetValue() == WmfMediaSession::RequestedTimeCurrent) ? TEXT("<current>") : *RequestedTime.GetValue().ToString();
-		const FString RequestedTopologyString = !RequestedTopology ? TEXT("<none>") : FString::Printf(TEXT("%p"), RequestedTopology.Get());
-
-		UE_LOG(LogWmfMedia, Verbose, TEXT("Session %p: Doing pending changes: RequestedRate: %s, RequestedTime: %s, RequestedTopology: %s"), this, *RequestedRateString, *RequestedTimeString, *RequestedTopologyString);
-	}
-	else
-	{
-		UE_LOG(LogWmfMedia, Verbose, TEXT("Session %p: Doing pending changes: none"), this);
+		return (EventType != MESessionStarted) ? EWorkItemResult::RetryLater : EWorkItemResult::Done;
 	}
 
-	PendingChanges = false;
+	// Actual setup...
 
-	// commit pending topology changes
-	if (RequestedTopology != NULL)
+	// Clear any topology enqueued to be set prior to the one we are about to set
+	// (this is async, but we omit a specific state change and wait as we execute the next async command unconditonally right afterwards - and we wait for that one)
+	HRESULT Result = MediaSession->ClearTopologies();
+	if (FAILED(Result))
 	{
-		TComPtr<IMFTopology> Topology = RequestedTopology;
-		RequestedTopology.Reset();
-
-		CommitTopology(Topology);
+		UE_LOG(LogWmfMedia, Verbose, TEXT("Session %p: Failed to clear queued topologies: %s"), this, *WmfMedia::ResultToString(Result));
+		return EWorkItemResult::Done;
+	}
 
-		if (PendingChanges)
+	// Set or clear topology...
+	if (Topology != nullptr)
+	{
+		Result = MediaSession->SetTopology(MFSESSION_SETTOPOLOGY_IMMEDIATE, Topology);
+		if (FAILED(Result))
 		{
-			return;
+			UE_LOG(LogWmfMedia, Verbose, TEXT("Session %p: Failed to set topology %p: %s"), Topology, this, *WmfMedia::ResultToString(Result));
+			return EWorkItemResult::Done;
 		}
 	}
-
-	// commit pending rate changes
-	if (RequestedRate.IsSet())
+	else
 	{
-		const float Rate = RequestedRate.GetValue();
-		RequestedRate.Reset();
-
-		CommitRate(Rate);
-
-		if (PendingChanges)
+		Result = MediaSession->SetTopology(MFSESSION_SETTOPOLOGY_CLEAR_CURRENT, Topology);
+		if (FAILED(Result))
 		{
-			return;
+			UE_LOG(LogWmfMedia, Verbose, TEXT("Session %p: Failed to set topology %p: %s"), Topology, this, *WmfMedia::ResultToString(Result));
+			return EWorkItemResult::Done;
 		}
 	}
 
-	// commit pending seeks/restarts
-	if (RequestedTime.IsSet())
-	{
-		const FTimespan Time = RequestedTime.GetValue();
-		RequestedTime.Reset();
+	UE_LOG(LogWmfMedia, Verbose, TEXT("Session %p: Committed topology %p"), this, Topology);
 
-		CommitTime(Time, bIsRequestedTimeSeek && !bIsRequestedTimeLoop);
-		bIsRequestedTimeLoop = false;
-		bIsRequestedTimeSeek = false;
-	}
+	return EWorkItemResult::RetryLater;
 }
 
 
@@ -1228,11 +1468,14 @@ void FWmfMediaSession::UpdateCharacteristics()
 	{
 		UE_LOG(LogWmfMedia, Verbose, TEXT("Session %p: Rate control ready"), this);
 
-		if (FAILED(RateControl->GetRate(FALSE, &SessionRate)))
+		float NewRate;
+		if (FAILED(RateControl->GetRate(FALSE, &NewRate)))
 		{
 			UE_LOG(LogWmfMedia, Verbose, TEXT("Session %p: Failed to initialize current rate"), this);
-			SessionRate = 1.0f; // the session's initial play rate is usually 1.0
+			NewRate = 1.0f; // the session's initial play rate is usually 1.0
 		}
+
+		SetSessionRate(NewRate);
 	}
 
 	Result = ::MFGetService(MediaSession, MF_RATE_CONTROL_SERVICE, IID_PPV_ARGS(&RateSupport));
@@ -1249,8 +1492,6 @@ void FWmfMediaSession::UpdateCharacteristics()
 	// cache rate control properties
 	if (RateSupport.IsValid())
 	{
-		CanScrub = SUCCEEDED(RateSupport->IsRateSupported(TRUE, 0.0f, NULL));
-
 		float MaxRate = 0.0f;
 		float MinRate = 0.0f;
 
@@ -1278,6 +1519,12 @@ void FWmfMediaSession::UpdateCharacteristics()
 			UnthinnedRates.Add(TRange<float>::Inclusive(MaxRate, MinRate));
 		}
 
+		// Check for scrubbing (playback at 0.0 speed)
+		/*
+		* We do NOT use IsRateSupported() as it seemns to produce incorrect results quite often (claiming incorrect support for 0.0, but advising 1.0 as best alternate rate)
+		*/
+		CanScrub = ThinnedRates.Contains(0.0f);
+
 		// When native out is enabled, the slowest rate will be greater than 0.0f
 		bool SupportsUnthinnedPause = SUCCEEDED(RateSupport->IsRateSupported(FALSE, 0.0f, NULL));
 		if (SupportsUnthinnedPause && !ThinnedRates.Contains(0.0f) && !UnthinnedRates.Contains(0.0f))
@@ -1288,6 +1535,39 @@ void FWmfMediaSession::UpdateCharacteristics()
 }
 
 
+void FWmfMediaSession::SetSessionRate(float Rate)
+{
+	if (SessionRate != Rate)
+	{
+		SessionRate = Rate;
+		if (Rate != 0.0f)
+		{
+			UnpausedSessionRate = Rate;
+		}
+	}
+}
+
+
+void::FWmfMediaSession::Flush()
+{
+	FScopeLock Lock(&CriticalSection);
+	UserIssuedSeeks = 0;
+}
+
+
+void FWmfMediaSession::RequestMoreVideoData()
+{
+	EnqueueWorkItem(FWorkItem([this](TComPtr<IMFMediaEvent> MediaEvent, MediaEventType EventType, HRESULT EventStatus, FBaseWorkItemState* BaseState) -> EWorkItemResult
+		{
+			TSharedPtr<FWmfMediaTracks, ESPMode::ThreadSafe> TracksPinned = Tracks.Pin();
+			if (TracksPinned.IsValid())
+			{
+				TracksPinned->RequestMoreVideoDataFromStreamSink();
+			}
+			return EWorkItemResult::Done;
+		}));
+}
+
 /* FWmfMediaSession callbacks
 *****************************************************************************/
 
@@ -1296,54 +1576,81 @@ void FWmfMediaSession::HandleError(HRESULT EventStatus)
 	UE_LOG(LogWmfMedia, Error, TEXT("An error occurred in the media session: %s"), *WmfMedia::ResultToString(EventStatus));
 
 	SessionState = EMediaState::Error;
-	DiscardPendingChanges();
 	MediaSession->Close();
+
+	ResetWorkItemQueue();
+}
+
+
+void FWmfMediaSession::HandlePresentationEnded()
+{
+	SessionState = EMediaState::Stopped;
 }
 
 
 void FWmfMediaSession::HandleSessionEnded()
 {
 	UE_LOG(LogWmfMedia, VeryVerbose, TEXT("FWmfMediaSession::HandleSessionEnded ShouldLoop:%d"), ShouldLoop);
-#if WMFMEDIA_PLAYER_VERSION >= 2
-	if (ShouldLoop)
-#endif // WMFMEDIA_PLAYER_VERSION >= 2
-	{
-		DeferredEvents.Enqueue(EMediaEvent::PlaybackEndReached);
-	}
 
 	SessionState = EMediaState::Stopped;
 
-	if (ShouldLoop)
+	TSharedPtr<FWmfMediaTracks, ESPMode::ThreadSafe> TracksPinned = Tracks.Pin();
+	if (TracksPinned.IsValid())
 	{
-		// loop back to beginning/end
-		RequestedTime = (SessionRate < 0.0f) ? CurrentDuration : FTimespan::Zero();
-		bIsRequestedTimeLoop = true;
-		bIsRequestedTimeSeek = false;
-		DoPendingChanges();
+		TracksPinned->SessionEnded();
 	}
-	else
+
+	// We only loop if we don't have a seek in progress right now
+	if (!bSeekActive)
 	{
-		LastTime = FTimespan::Zero();
-		RequestedRate.Reset();
-#if WMFMEDIA_PLAYER_VERSION >= 2
-		bIsWaitingForEnd = true;
-#endif // WMFMEDIA_PLAYER_VERSION >= 2
+		if (!ShouldLoop)
+		{
+			DeferredEvents.Enqueue(EMediaEvent::PlaybackEndReached);
+		}
+
+		if (ShouldLoop)
+		{
+			// loop back to beginning/end
+			JamWorkItem(FWorkItem([this](TComPtr<IMFMediaEvent> MediaEvent, MediaEventType EventType, HRESULT EventStatus, FBaseWorkItemState* BaseState) -> EWorkItemResult
+				{
+					auto State = static_cast<FTimeWorkItemState*>(BaseState);
+
+					if (State->State == FTimeWorkItemState::EState::WaitDone)
+					{
+						return (EventType == MESessionStarted) ? EWorkItemResult::Done : EWorkItemResult::RetryLater;
+					}
+
+					if (CommitTime((UnpausedSessionRate < 0.0f) ? CurrentDuration : FTimespan::Zero(), false))
+					{
+						TSharedPtr<FWmfMediaTracks, ESPMode::ThreadSafe> TracksPinned = Tracks.Pin();
+						if (TracksPinned.IsValid())
+						{
+							TracksPinned->LoopStarted(UnpausedSessionRate);
+						}
+
+						State->State = FTimeWorkItemState::EState::WaitDone;
+						return EWorkItemResult::RetryLater;
+					}
+
+					return EWorkItemResult::Done;
+				}, MakeShared<FTimeWorkItemState, ESPMode::ThreadSafe>()));
+		}
+		else
+		{
+			LastTime = FTimespan::Zero();
+			bIsWaitingForEnd = true;
+			ResetWorkItemQueue();
+		}
 	}
 }
 
 
 void FWmfMediaSession::HandleSessionPaused(HRESULT EventStatus)
 {
-	if (SUCCEEDED(EventStatus))
+	if (RateControl == nullptr || !CanScrub)
 	{
-		SessionState = EMediaState::Paused;
 		DeferredEvents.Enqueue(EMediaEvent::PlaybackSuspended);
-
-		DoPendingChanges();
-	}
-	else
-	{
-		DiscardPendingChanges();
+		SessionState = EMediaState::Paused;
 	}
 }
 
@@ -1360,23 +1667,20 @@ void FWmfMediaSession::HandleSessionRateChanged(HRESULT EventStatus, IMFMediaEve
 
 		if (SUCCEEDED(Result) && (Value.vt == VT_R4))
 		{
-			SessionRate = Value.fltVal;
+			SetSessionRate(Value.fltVal);
 		}
 	}
 	else if (RateControl != NULL)
 	{
+
 		BOOL Thin = FALSE;
-		RateControl->GetRate(&Thin, &SessionRate);
+		float Rate = 0.0f; // quiet down SA validation
+		if (RateControl->GetRate(&Thin, &Rate) != S_OK)
+		{
+			Rate = 1.0f;
+		}
+		SetSessionRate(Rate);
 	}
-
-	DoPendingChanges();
-}
-
-
-void FWmfMediaSession::HandleSessionScrubSampleComplete()
-{
-//	DeferredEvents.Enqueue(EMediaEvent::SeekCompleted);
-//	DoPendingChanges();
 }
 
 
@@ -1384,56 +1688,41 @@ void FWmfMediaSession::HandleSessionStarted(HRESULT EventStatus)
 {
 	if (SUCCEEDED(EventStatus))
 	{
-#if WMFMEDIA_PLAYER_VERSION >= 2
 		bIsWaitingForEnd = false;
-#endif // WMFMEDIA_PLAYER_VERSION >= 2
 
-		if ((SessionState == EMediaState::Paused) && (SessionRate == 0.0f))
+		if (RateControl == NULL)
 		{
-			SessionState = EMediaState::Playing;
-
-			// Note: pending changes will be processed in MESessionScrubSampleComplete
-
-			DeferredEvents.Enqueue(EMediaEvent::SeekCompleted);
-			DoPendingChanges();
+			SetSessionRate(1.0f);
 		}
 		else
 		{
-			if (RateControl == NULL)
-			{
-				SessionRate = 1.0f;
-			}
-			else
+			float Rate = 0.0f; // quiet down SA validation
+			if (RateControl->GetRate(NULL, &Rate) != S_OK)
 			{
-				BOOL Thin = FALSE;
-				RateControl->GetRate(&Thin, &SessionRate);
+				Rate = 1.0f;
 			}
+			SetSessionRate(Rate);
+		}
 
-			if (SessionState == EMediaState::Playing)
-			{
-				DeferredEvents.Enqueue(EMediaEvent::SeekCompleted);
-			}
-			else
-			{
-				SessionState = EMediaState::Playing;
+		if (SessionRate != 0.0f)
+		{
+			SessionState = EMediaState::Playing;
+			DeferredEvents.Enqueue(EMediaEvent::PlaybackResumed);
+		}
+		else
+		{
+			// We are pausing (or "scrubbing" in WMF terms - hence this is a playback mode, not real "pause")
+			SessionState = EMediaState::Paused;
+			DeferredEvents.Enqueue(EMediaEvent::PlaybackSuspended);
+		}
 
-				if (SessionRate == 0.0f)
-				{
-					RequestedRate = 0.0f;
-					PendingChanges = true;
-				}
-				else
-				{
-					DeferredEvents.Enqueue(EMediaEvent::PlaybackResumed);
-				}
-			}
-			
-			DoPendingChanges();
+		UE_LOG(LogWmfMedia, VeryVerbose, TEXT("Session - Started: R=%f SState=%d"), SessionRate, int(SessionState));
+
+		if (bSeekActive && (SessionState == EMediaState::Playing || SessionState == EMediaState::Paused))
+		{
+			DeferredEvents.Enqueue(EMediaEvent::SeekCompleted);
 		}
-	}
-	else
-	{
-		DiscardPendingChanges();
+		bSeekActive = false;
 	}
 }
 
@@ -1444,11 +1733,6 @@ void FWmfMediaSession::HandleSessionStopped(HRESULT EventStatus)
 	{
 		SessionState = EMediaState::Stopped;
 		DeferredEvents.Enqueue(EMediaEvent::PlaybackSuspended);
-		DoPendingChanges();
-	}
-	else
-	{
-		DiscardPendingChanges();
 	}
 }
 
@@ -1462,15 +1746,6 @@ void FWmfMediaSession::HandleSessionTopologySet(HRESULT EventStatus, IMFMediaEve
 		if (SUCCEEDED(Result))
 		{
 			UE_LOG(LogWmfMedia, Verbose, TEXT("Session %p: Topology %p set as current"), this, CurrentTopology.Get());
-
-			if (SessionState != EMediaState::Preparing)
-			{
-				// Note: track and format changes won't send an MF_TOPOSTATUS_READY event
-				// until playback is restarted, so we do pending changes here instead.
-
-				DoPendingChanges();
-			}
-
 			return;
 		}
 
@@ -1485,8 +1760,6 @@ void FWmfMediaSession::HandleSessionTopologySet(HRESULT EventStatus, IMFMediaEve
 		SessionState = EMediaState::Error;
 		DeferredEvents.Enqueue(EMediaEvent::MediaOpenFailed);
 	}
-
-	DiscardPendingChanges();
 }
 
 
@@ -1528,8 +1801,7 @@ void FWmfMediaSession::HandleSessionTopologyStatus(HRESULT EventStatus, IMFMedia
 
 	if (SessionState == EMediaState::Error)
 	{
-		DiscardPendingChanges();
-
+		ResetWorkItemQueue();
 		return;
 	}
 
@@ -1544,8 +1816,7 @@ void FWmfMediaSession::HandleSessionTopologyStatus(HRESULT EventStatus, IMFMedia
 			DeferredEvents.Enqueue(EMediaEvent::MediaOpenFailed);
 		}
 
-		DiscardPendingChanges();
-
+		ResetWorkItemQueue();
 		return;
 	}
 
@@ -1588,19 +1859,26 @@ void FWmfMediaSession::HandleSessionTopologyStatus(HRESULT EventStatus, IMFMedia
 	{
 		// Note: when paused, the new topology won't apply until the next session start,
 		// so we request a scrub to the current time in order to update the video frame.
+		JamWorkItem(FWorkItem([this](TComPtr<IMFMediaEvent> MediaEvent, MediaEventType EventType, HRESULT EventStatus, FBaseWorkItemState* BaseState) -> EWorkItemResult
+			{
+				auto State = static_cast<FTimeWorkItemState*>(BaseState);
 
-		if (!RequestedTime.IsSet())
-		{
-			UE_LOG(LogWmfMedia, Verbose, TEXT("Session %p: Requesting scrub after topology change"), this);
+				UE_LOG(LogWmfMedia, Verbose, TEXT("Session %p: 'Re-Scrubbing' after topology change [SState=%d X=%d E=%d]"), this, int(State->State), EventType == MESessionStarted, int(EventType));
 
-			RequestedTime = WmfMediaSession::RequestedTimeCurrent;
-			bIsRequestedTimeLoop = false;
-			bIsRequestedTimeSeek = false;
-			PendingChanges = true;
-		}
-	}
+				if (State->State == FTimeWorkItemState::EState::WaitDone)
+				{
+					return (EventType == MESessionStarted) ? EWorkItemResult::Done : EWorkItemResult::RetryLater;
+				}
+
+				if (CommitTime(WmfMediaSession::RequestedTimeCurrent, false))
+				{
+					State->State = FTimeWorkItemState::EState::WaitDone;
+					return EWorkItemResult::RetryLater;
+				}
 
-	DoPendingChanges();
+				return EWorkItemResult::Done;
+			}, MakeShared<FTimeWorkItemState, ESPMode::ThreadSafe>()));
+	}
 }
 
 
diff --git a/Engine/Plugins/Media/WmfMedia/Source/WmfMedia/Private/Wmf/WmfMediaSession.h b/Engine/Plugins/Media/WmfMedia/Source/WmfMedia/Private/Wmf/WmfMediaSession.h
--- a/Engine/Plugins/Media/WmfMedia/Source/WmfMedia/Private/Wmf/WmfMediaSession.h
+++ b/Engine/Plugins/Media/WmfMedia/Source/WmfMedia/Private/Wmf/WmfMediaSession.h
@@ -63,14 +63,12 @@ public:
 	 */
 	void GetEvents(TArray<EMediaEvent>& OutEvents);
 
-#if WMFMEDIA_PLAYER_VERSION >= 2
 	/**
 	 * Set which tracks object is being used by the player.
 	 *
 	 * @param InTracks Tracks object.
 	 */
 	void SetTracks(TSharedPtr<FWmfMediaTracks, ESPMode::ThreadSafe> InTracks);
-#endif // WMFMEDIA_PLAYER_VERSION >= 2
 
 	/**
 	 * Initialize the media session.
@@ -123,7 +121,93 @@ public:
 	STDMETHODIMP QueryInterface(REFIID RefID, void** Object);
 	STDMETHODIMP_(ULONG) Release();
 
+public:
+	void Flush();
+
+	void RequestMoreVideoData();
+
+	enum class EWorkItemResult
+	{
+		Done,
+		RetryLater
+	};
+
+	struct FBaseWorkItemState
+	{
+		virtual ~FBaseWorkItemState() = default;
+	};
+
+	class FWorkItem
+	{
+	public:
+		typedef TFunction<EWorkItemResult(TComPtr<IMFMediaEvent> MediaEvent, MediaEventType EventType, HRESULT EventStatus, FBaseWorkItemState* BaseState)>	FWorkLoadType;
+
+		FWorkItem() = default;
+		FWorkItem(FWorkLoadType&& InWorkLoad, TSharedPtr<FBaseWorkItemState, ESPMode::ThreadSafe> InState = nullptr) : WorkLoad(MoveTemp(InWorkLoad)), State(InState) {}
+
+		FWorkLoadType WorkLoad;
+		TSharedPtr<FBaseWorkItemState, ESPMode::ThreadSafe> State;
+
+		bool IsValid() const
+		{
+			return (bool)WorkLoad;
+		}
+
+		void Reset()
+		{
+			WorkLoad.Reset();
+			State.Reset();
+		}
+	};
+
+public:
+	void EnqueueWorkItem(FWorkItem&& WorkItem);
+	void JamWorkItem(FWorkItem&& WorkItem);
+
 protected:
+	struct FTopologyWorkItemState : public FBaseWorkItemState
+	{
+		enum class EState
+		{
+			Begin,
+			Ready,
+			ReadyNeedsRestart,
+			WaitRestart,
+		};
+
+		EState State = EState::Begin;
+	};
+
+	struct FRateWorkItemState : public FBaseWorkItemState
+	{
+		enum class EState
+		{
+			Begin,
+			Ready,
+			ReadyNeedsRestart,
+			WaitForStopThenSetAndRestart,
+			WaitForPauseThenSetAndRestart,
+			WaitForSet,
+			WaitForSetAndRestart,
+			WaitForRestart,
+			WaitForPause,
+			WaitForStart
+		};
+
+		EState State = EState::Begin;
+		FTimespan LastTime;
+	};
+
+	struct FTimeWorkItemState : public FBaseWorkItemState
+	{
+		enum class EState
+		{
+			Begin,
+			WaitDone,
+		};
+
+		EState State = EState::Begin;
+	};
 
 	/**
 	 * Commit the specified play rate.
@@ -133,7 +217,7 @@ protected:
 	 * @param Rate The play rate to commit.
 	 * @see CommitTime, CommitTopology
 	 */
-	bool CommitRate(float Rate);
+	EWorkItemResult CommitRate(float Rate, FRateWorkItemState* State, MediaEventType EventType);
 
 	/**
 	 * Commit the specified play position.
@@ -151,25 +235,16 @@ protected:
 	 * @param Topology The topology to set.
 	 * @see CommitRate, CommitTime
 	 */
-	bool CommitTopology(IMFTopology* Topology);
-
-	/**
-	 * Discard all pending state changes.
-	 *
-	 * @see DoPendingChanges
-	 */
-	void DiscardPendingChanges();
-
-	/**
-	* Applies any pending state changes.
-	*
-	* @see DiscardPendingChanges
-	*/
-	void DoPendingChanges();
+	EWorkItemResult CommitTopology(IMFTopology* Topology, FTopologyWorkItemState* State, MediaEventType EventType);
 
 	/** Get the latest characteristics from the current media source. */
 	void UpdateCharacteristics();
 
+	void ExecuteNextWorkItems(TComPtr<IMFMediaEvent> MediaEvent, MediaEventType EventType, HRESULT EventStatus);
+	void ResetWorkItemQueue();
+
+	void SetSessionRate(float Rate);
+
 private:
 
 	/** Callback for the MEError event. */
@@ -178,15 +253,14 @@ private:
 	/** Callback for the MESessionEnded event. */
 	void HandleSessionEnded();
 
+	void HandlePresentationEnded();
+
 	/** Callback for the MESessionPaused event. */
 	void HandleSessionPaused(HRESULT EventStatus);
 
 	/** Callback for the MESessionRateChanged event. */
 	void HandleSessionRateChanged(HRESULT EventStatus, IMFMediaEvent& Event);
 
-	/** Callback for the MESessionScrubSampleComplete event. */
-	void HandleSessionScrubSampleComplete();
-
 	/** Callback for the MESessionStarted event. */
 	void HandleSessionStarted(HRESULT EventStatus);
 
@@ -233,9 +307,6 @@ private:
 	/** The last play head position before playback was stopped. */
 	FTimespan LastTime;
 
-	/** Whether one or more state changes are pending. */
-	bool PendingChanges;
-
 	/** The media session's clock. */
 	TComPtr<IMFPresentationClock> PresentationClock;
 
@@ -254,18 +325,20 @@ private:
 	/** Deferred playback topology to set. */
 	TComPtr<IMFTopology> RequestedTopology;
 
-	/** Deferred play time change value (MinValue = no change, MaxValue = current time). */
-	TOptional<FTimespan> RequestedTime;
-
 	/** If true then RequestedTime is due to the video looping. */
 	bool bIsRequestedTimeLoop;
 
 	/** If true then RequestedTime is due to the video seeking. */
 	bool bIsRequestedTimeSeek;
 
+	float LastSetRate;
+
 	/** The session's internal playback rate (not necessarily the same as GetRate). */
 	float SessionRate;
 
+	/** The session's last non-zero internal playback rate. */
+	float UnpausedSessionRate;
+
 	/** The session's current state (not necessarily the same as GetState). */
 	EMediaState SessionState;
 
@@ -284,13 +357,27 @@ private:
 	/** The unthinned play rates that the current media session supports. */
 	TRangeSet<float> UnthinnedRates;
 
-#if WMFMEDIA_PLAYER_VERSION >= 2
 	/** Pointer to the tracks from the player. */
 	TWeakPtr<FWmfMediaTracks, ESPMode::ThreadSafe> Tracks;
 
 	/** True if we are waiting for the tracks to send out its last samples before we "end". */
 	bool bIsWaitingForEnd;
-#endif // WMFMEDIA_PLAYER_VERSION >= 2
+
+	/** Number of user initiated seek calls accumulated for next actual seek executed */
+	uint32 UserIssuedSeeks;
+
+	/** True if seek is logically in progress */
+	bool bSeekActive;
+
+	/** Queue of work jobs for the player */
+	TArray<FWorkItem> WorkItemQueue;
+
+	/** The work item currently in progress (if valid) */
+	FWorkItem CurrentWorkItem;
+
+	/** Number of seek wokr-items in the queue currently */
+	int32 NumSeeksInWorkQueue;
+
 };
 
 
diff --git a/Engine/Plugins/Media/WmfMedia/Source/WmfMedia/Private/Wmf/WmfMediaSink.cpp b/Engine/Plugins/Media/WmfMedia/Source/WmfMedia/Private/Wmf/WmfMediaSink.cpp
--- a/Engine/Plugins/Media/WmfMedia/Source/WmfMedia/Private/Wmf/WmfMediaSink.cpp
+++ b/Engine/Plugins/Media/WmfMedia/Source/WmfMedia/Private/Wmf/WmfMediaSink.cpp
@@ -184,18 +184,12 @@ STDMETHODIMP FWmfMediaSink::GetCharacteristics(__RPC__out DWORD* pdwCharacterist
 		return E_POINTER;
 	}
 
-	FScopeLock Lock(&CriticalSection);
-
 	if (!StreamSink.IsValid())
 	{
 		return MF_E_SHUTDOWN;
 	}
 
-#if WMFMEDIA_PLAYER_VERSION == 1
-	*pdwCharacteristics = MEDIASINK_FIXED_STREAMS | MEDIASINK_CAN_PREROLL;
-#else  // WMFMEDIA_PLAYER_VERSION == 1
-	*pdwCharacteristics = MEDIASINK_FIXED_STREAMS | MEDIASINK_CAN_PREROLL | MEDIASINK_RATELESS;
-#endif // WMFMEDIA_PLAYER_VERSION == 1
+	*pdwCharacteristics = MEDIASINK_FIXED_STREAMS | MEDIASINK_RATELESS;
 
 	return S_OK;
 }
@@ -331,19 +325,40 @@ STDMETHODIMP FWmfMediaSink::SetPresentationClock(__RPC__in_opt IMFPresentationCl
 	// Register ourselves to get state notifications from the new clock.
 	if (pPresentationClock != NULL)
 	{
-		const HRESULT Result = pPresentationClock->AddClockStateSink(this);
+		TComPtr<IMFRateControl> RateControl;
+
+		HRESULT Result = pPresentationClock->QueryInterface(IID_PPV_ARGS(&RateControl));
+ 		if (FAILED(Result))
+		{
+			UE_LOG(LogWmfMedia, Error, TEXT("Failed to retrieve rate interface from presentation clock: %s"), *WmfMedia::ResultToString(Result));
+			return Result;
+		}
 
+		float Rate;
+		Result = RateControl->GetRate(NULL, &Rate);
+		if (FAILED(Result))
+		{
+			UE_LOG(LogWmfMedia, Error, TEXT("Failed to retrieve rate from presentation clock: %s"), *WmfMedia::ResultToString(Result));
+			return Result;
+		}
+
+		// Make sure the sink knows the current clock rate (the presentation clock may only be known after initial rate setup is already done)
+		StreamSink->SetClockRate(Rate);
+
+		Result = pPresentationClock->AddClockStateSink(this);
 		if (FAILED(Result))
 		{
 			UE_LOG(LogWmfMedia, Error, TEXT("Failed to add media sink to presentation clock: %s"), *WmfMedia::ResultToString(Result));
 			return Result;
 		}
 	}
+	else
+	{
+		StreamSink->SetClockRate(-1.0f);
+	}
 
 	PresentationClock = pPresentationClock;
 
-	StreamSink->SetPresentationClock(PresentationClock);
-
 	return S_OK;
 }
 
@@ -368,22 +383,6 @@ STDMETHODIMP FWmfMediaSink::Shutdown()
 }
 
 
-/* IMFMediaSinkPreroll interface
- *****************************************************************************/
-
-STDMETHODIMP FWmfMediaSink::NotifyPreroll(MFTIME hnsUpcomingStartTime)
-{
-	FScopeLock Lock(&CriticalSection);
-
-	if (!StreamSink.IsValid())
-	{
-		return MF_E_SHUTDOWN;
-	}
-
-	return StreamSink->Preroll();
-}
-
-
 /* IMFRateSupport interface
  *****************************************************************************/
 
diff --git a/Engine/Plugins/Media/WmfMedia/Source/WmfMedia/Private/Wmf/WmfMediaSink.h b/Engine/Plugins/Media/WmfMedia/Source/WmfMedia/Private/Wmf/WmfMediaSink.h
--- a/Engine/Plugins/Media/WmfMedia/Source/WmfMedia/Private/Wmf/WmfMediaSink.h
+++ b/Engine/Plugins/Media/WmfMedia/Source/WmfMedia/Private/Wmf/WmfMediaSink.h
@@ -21,7 +21,6 @@ class FWmfMediaSink
 	: public IMFClockStateSink
 	, public IMFGetService
 	, public IMFMediaSink
-	, public IMFMediaSinkPreroll
 	, public IMFRateSupport
 {
 public:
@@ -69,12 +68,6 @@ public:
 	STDMETHODIMP SetPresentationClock(__RPC__in_opt IMFPresentationClock* pPresentationClock);
 	STDMETHODIMP Shutdown();
 
-public:
-
-	//~ IMFMediaSinkPreroll interface
-
-	STDMETHODIMP NotifyPreroll(MFTIME hnsUpcomingStartTime);
-
 public:
 
 	//~ IMFRateSupport interface
diff --git a/Engine/Plugins/Media/WmfMedia/Source/WmfMedia/Private/Wmf/WmfMediaStreamSink.cpp b/Engine/Plugins/Media/WmfMedia/Source/WmfMedia/Private/Wmf/WmfMediaStreamSink.cpp
--- a/Engine/Plugins/Media/WmfMedia/Source/WmfMedia/Private/Wmf/WmfMediaStreamSink.cpp
+++ b/Engine/Plugins/Media/WmfMedia/Source/WmfMedia/Private/Wmf/WmfMediaStreamSink.cpp
@@ -16,6 +16,8 @@
 #include "Windows/AllowWindowsPlatformTypes.h"
 #include <d3d11.h>
 
+static constexpr uint32 kMaxFramesInOutputQueue = 6;
+
 /* FWmfMediaStreamSink static functions
  *****************************************************************************/
 
@@ -41,14 +43,18 @@ bool FWmfMediaStreamSink::Create(const GUID& MajorType, TComPtr<FWmfMediaStreamS
 FWmfMediaStreamSink::FWmfMediaStreamSink(const GUID& InMajorType, DWORD InStreamId)
 	: Owner(nullptr)
 	, Decoder(nullptr)
-	, Prerolling(false)
 	, RefCount(0)
 	, StreamId(InStreamId)
 	, StreamType(InMajorType)
-	, ClockRate(1.0f)
-	, WaitTimer(nullptr)
+	, ClockRate(-1.0f)
 	, VideoSamplePool(nullptr)
 	, VideoSampleQueue(nullptr)
+	, NumRequestsInFlight(0)
+	, NumRequestsInFlightWithSystem(0)
+	, bSeekInProgress(false)
+	, bScrubbingInProgress(false)
+	, bSystemIsScrubbing(false)
+	, bSegmentEnded(true)
 	, bShowSubTypeErrorMessage(true)
 {
 	UE_LOG(LogWmfMedia, Verbose, TEXT("StreamSink %p: Created with stream type %s"), this, *WmfMedia::MajorTypeToString(StreamType));
@@ -66,38 +72,6 @@ FWmfMediaStreamSink::~FWmfMediaStreamSink()
 /* FWmfMediaStreamSink interface
  *****************************************************************************/
 
-bool FWmfMediaStreamSink::GetNextSample(TComPtr<IMFSample>& OutSample)
-{
-	FScopeLock Lock(&CriticalSection);
-
-#if WMFMEDIA_PLAYER_VERSION == 1
-	while (SampleQueue.Num())
-	{
-		FQueuedSample QueuedSample = SampleQueue.Pop();
-#else // WMFMEDIA_PLAYER_VERSION == 1
-	FQueuedSample QueuedSample;
-	if (SampleQueue.Dequeue(QueuedSample))
-	{
-#endif // WMFMEDIA_PLAYER_VERSION == 1
-		if (QueuedSample.Sample.IsValid())
-		{
-			OutSample = QueuedSample.Sample;
-			return true;
-		}
-		else
-		{
-			// process pending marker
-			QueueEvent(MEStreamSinkMarker, GUID_NULL, S_OK, QueuedSample.MarkerContext);
-			PropVariantClear(QueuedSample.MarkerContext);
-			delete QueuedSample.MarkerContext;
-			UE_LOG(LogWmfMedia, Verbose, TEXT("StreamSink %p: Processed marker (%s)"), this, *WmfMedia::MarkerTypeToString(QueuedSample.MarkerType));
-		}
-	}
-
-	return false;
-}
-
-
 bool FWmfMediaStreamSink::Initialize(FWmfMediaSink& InOwner)
 {
 	FScopeLock Lock(&CriticalSection);
@@ -129,23 +103,6 @@ HRESULT FWmfMediaStreamSink::Pause()
 }
 
 
-HRESULT FWmfMediaStreamSink::Preroll()
-{
-	FScopeLock Lock(&CriticalSection);
-
-	if (!EventQueue.IsValid())
-	{
-		return MF_E_SHUTDOWN;
-	}
-
-	Prerolling = true;
-
-	UE_LOG(LogWmfMedia, VeryVerbose, TEXT("StreamSink %p:Preroll Request Sample"), this);
-
-	return QueueEvent(MEStreamSinkRequestSample, GUID_NULL, S_OK, NULL);
-}
-
-
 HRESULT FWmfMediaStreamSink::Restart()
 {
 	HRESULT Result = QueueEvent(MEStreamSinkStarted, GUID_NULL, S_OK, NULL);
@@ -157,7 +114,7 @@ HRESULT FWmfMediaStreamSink::Restart()
 
 	UE_LOG(LogWmfMedia, VeryVerbose, TEXT("StreamSink %p:Restart Request Sample"), this);
 
-	return QueueEvent(MEStreamSinkRequestSample, GUID_NULL, S_OK, NULL);
+	return S_OK;
 }
 
 
@@ -172,31 +129,15 @@ void FWmfMediaStreamSink::Shutdown()
 	}
 
 	CurrentMediaType.Reset();
-
-	CloseTimer();
 }
 
 
 HRESULT FWmfMediaStreamSink::Start()
 {
 	FScopeLock Lock(&CriticalSection);
-	UE_LOG(LogWmfMedia, VeryVerbose, TEXT("StreamSink::Start Rate:%f"), ClockRate);
-	if (WaitTimer == nullptr)
-	{
-		// Set a high the timer resolution (ie, short timer period).
-		timeBeginPeriod(1);
 
-		// create the waitable timer
-		WaitTimer = CreateWaitableTimer(NULL, FALSE, NULL);
-		if (WaitTimer == nullptr)
-		{
-			HRESULT Result = HRESULT_FROM_WIN32(GetLastError());
-			if (FAILED(Result))
-			{
-				return Result;
-			}
-		}
-	}
+	bScrubbingInProgress = (ClockRate == 0.0f);
+	bSystemIsScrubbing = bScrubbingInProgress;
 
 	HRESULT Result = QueueEvent(MEStreamSinkStarted, GUID_NULL, S_OK, NULL);
 
@@ -207,34 +148,39 @@ HRESULT FWmfMediaStreamSink::Start()
 
 	UE_LOG(LogWmfMedia, VeryVerbose, TEXT("StreamSink %p:Start Request Sample"), this);
 
-	return QueueEvent(MEStreamSinkRequestSample, GUID_NULL, S_OK, NULL);
+	bSegmentEnded = false;
+	bSeekInProgress = false;
+
+	return S_OK;
 }
 
 
 HRESULT FWmfMediaStreamSink::Stop()
 {
+	HRESULT Res;
+	TFunction<void()> TempAction;
+	{
 	UE_LOG(LogWmfMedia, VeryVerbose, TEXT("StreamSink::Stop"));
-#if WMFMEDIA_PLAYER_VERSION == 1
-	Flush();
-#endif // WMFMEDIA_PLAYER_VERSION == 1
 
 	FScopeLock Lock(&CriticalSection);
 
-	// Restore the timer resolution.
-	timeEndPeriod(1);
+	NumRequestsInFlight -= NumRequestsInFlightWithSystem;
+	NumRequestsInFlightWithSystem = 0;
 
-	CloseTimer();
+	Res = QueueEvent(MEStreamSinkStopped, GUID_NULL, S_OK, NULL);
 
-	return QueueEvent(MEStreamSinkStopped, GUID_NULL, S_OK, NULL);
-}
+	if ((NumRequestsInFlight == 0) && OnIdleAction)
+	{
+		TempAction = MoveTemp(OnIdleAction);
+	}
+	}
 
-void FWmfMediaStreamSink::CloseTimer()
-{
-	if (WaitTimer != nullptr)
+	if (TempAction)
 	{
-		CloseHandle(WaitTimer);
-		WaitTimer = nullptr;
+		TempAction();
 	}
+
+	return Res;
 }
 
 /* IMFGetService interface
@@ -271,7 +217,19 @@ STDMETHODIMP FWmfMediaStreamSink::EndGetEvent(IMFAsyncResult* pResult, IMFMediaE
 		return MF_E_SHUTDOWN;
 	}
 
-	return EventQueue->EndGetEvent(pResult, ppEvent);
+	HRESULT Result = EventQueue->EndGetEvent(pResult, ppEvent);
+
+	if (Result == S_OK)
+	{
+		MediaEventType Type;
+		(*ppEvent)->GetType(&Type);
+		if (Type == MEStreamSinkRequestSample)
+		{
+			++NumRequestsInFlightWithSystem;
+		}
+	}
+
+	return Result;
 }
 
 
@@ -289,7 +247,20 @@ STDMETHODIMP FWmfMediaStreamSink::GetEvent(DWORD dwFlags, IMFMediaEvent** ppEven
 		TempQueue = EventQueue;
 	}
 
-	return TempQueue->GetEvent(dwFlags, ppEvent);
+	HRESULT Result = TempQueue->GetEvent(dwFlags, ppEvent);
+
+	if (Result == S_OK)
+	{
+		MediaEventType Type;
+		(*ppEvent)->GetType(&Type);
+		if (Type == MEStreamSinkRequestSample)
+		{
+			FScopeLock Lock(&CriticalSection);
+			++NumRequestsInFlightWithSystem;
+		}
+	}
+
+	return Result;
 }
 
 
@@ -534,42 +505,12 @@ STDMETHODIMP FWmfMediaStreamSink::SetCurrentMediaType(IMFMediaType* pMediaType)
 
 STDMETHODIMP FWmfMediaStreamSink::Flush()
 {
-	FScopeLock Lock(&CriticalSection);
-
 	if (!EventQueue.IsValid())
 	{
 		return MF_E_SHUTDOWN;
 	}
 
-	UE_LOG(LogWmfMedia, Verbose, TEXT("StreamSink %p: Flushing samples & markers Rate:%f"), this, ClockRate);
-
-#if WMFMEDIA_PLAYER_VERSION == 1
-	while (SampleQueue.Num())
-	{
-		FQueuedSample QueuedSample = SampleQueue.Pop();
-#else // WMFMEDIA_PLAYER_VERSION == 1
-	FQueuedSample QueuedSample;
-	while (SampleQueue.Dequeue(QueuedSample))
-	{
-#endif // WMFMEDIA_PLAYER_VERSION == 1
-		if (QueuedSample.Sample.IsValid())
-		{
-			continue;
-		}
-
-		// notify WMF that flushed markers haven't been processed
-		QueueEvent(MEStreamSinkMarker, GUID_NULL, E_ABORT, QueuedSample.MarkerContext);
-		PropVariantClear(QueuedSample.MarkerContext);
-		delete QueuedSample.MarkerContext;
-	}
-
-#if WMFMEDIA_PLAYER_VERSION >= 2
-	// If the rate is 0 then get rid of the old samples, otherwise they might linger and we don't want them.
-	if (ClockRate == 0.0f)
-	{
-		VideoSampleQueue->RequestFlush();
-	}
-#endif // WMFMEDIA_PLAYER_VERSION >= 2
+	VideoSampleQueue->RequestFlush();
 
 	return S_OK;
 }
@@ -582,8 +523,6 @@ STDMETHODIMP FWmfMediaStreamSink::GetIdentifier(__RPC__out DWORD* pdwIdentifier)
 		return E_POINTER;
 	}
 
-	FScopeLock Lock(&CriticalSection);
-
 	if (!EventQueue.IsValid())
 	{
 		return MF_E_SHUTDOWN;
@@ -602,8 +541,6 @@ STDMETHODIMP FWmfMediaStreamSink::GetMediaSink(__RPC__deref_out_opt IMFMediaSink
 		return E_POINTER;
 	}
 
-	FScopeLock Lock(&CriticalSection);
-
 	if (!EventQueue.IsValid())
 	{
 		return MF_E_SHUTDOWN;
@@ -623,8 +560,6 @@ STDMETHODIMP FWmfMediaStreamSink::GetMediaTypeHandler(__RPC__deref_out_opt IMFMe
 		return E_POINTER;
 	}
 
-	FScopeLock Lock(&CriticalSection);
-
 	if (!EventQueue.IsValid())
 	{
 		return MF_E_SHUTDOWN;
@@ -643,8 +578,12 @@ STDMETHODIMP FWmfMediaStreamSink::PlaceMarker(MFSTREAMSINK_MARKER_TYPE eMarkerTy
 		return MF_E_SHUTDOWN;
 	}
 
-	UE_LOG(LogWmfMedia, Verbose, TEXT("StreamSink %p: Placing marker (%s)"), this, *WmfMedia::MarkerTypeToString(eMarkerType));
+	if (eMarkerType == MFSTREAMSINK_MARKER_ENDOFSEGMENT)
+	{
+		bSegmentEnded = true;
+	}
 
+	// Pass on all markers via event...
 	PROPVARIANT* MarkerContext = new PROPVARIANT;
 
 	if (pvarContextValue != NULL)
@@ -660,20 +599,77 @@ STDMETHODIMP FWmfMediaStreamSink::PlaceMarker(MFSTREAMSINK_MARKER_TYPE eMarkerTy
 		}
 	}
 
-#if WMFMEDIA_PLAYER_VERSION == 1
-	SampleQueue.Add({ eMarkerType, MarkerContext, nullptr });
-#else // WMFMEDIA_PLAYER_VERSION == 1
-	SampleQueue.Enqueue({ eMarkerType, MarkerContext, nullptr });
-#endif // WMFMEDIA_PLAYER_VERSION == 1
+	QueueEvent(MEStreamSinkMarker, GUID_NULL, S_OK, MarkerContext);
+	PropVariantClear(MarkerContext);
+	delete MarkerContext;
+
+	return S_OK;
+}
+
+
+void FWmfMediaStreamSink::OnSeekStarted()
+{
+	bSeekInProgress = true;
+}
 
-	TComPtr<IMFSample> NextSample;
-	if (GetNextSample(NextSample))
+
+void FWmfMediaStreamSink::OnSessionEnded()
+{
+	TFunction<void()> TempAction;
 	{
-		// process next samples
-		ScheduleWaitForNextSample(NextSample);
+	FScopeLock Lock(&CriticalSection);
+	
+	NumRequestsInFlight -= NumRequestsInFlightWithSystem;
+	NumRequestsInFlightWithSystem = 0;
+
+	if ((NumRequestsInFlight == 0) && OnIdleAction)
+	{
+		TempAction = MoveTemp(OnIdleAction);
+	}
 	}
 
-	return S_OK;
+	if (TempAction)
+	{
+		TempAction();
+	}
+}
+
+
+void FWmfMediaStreamSink::OnRequestMoreData()
+{
+	FScopeLock Lock(&CriticalSection);
+
+	if (!OnIdleAction)
+	{
+		/*
+		* - We need the event queue
+		* - We do not accept new request if we just received a "end of segment" & no new start since
+		* - The output queue must not be full
+		* - If scrubbing, we only accept one request once per "start with zero" phase
+		*/
+		if (EventQueue.IsValid() && !bSegmentEnded && !IsVideoSampleQueueFull() && (!bSystemIsScrubbing || (bScrubbingInProgress && NumRequestsInFlight == 0)))
+		{
+			QueueEvent(MEStreamSinkRequestSample, GUID_NULL, S_OK, NULL);
+			++NumRequestsInFlight;
+		}
+	}
+}
+
+
+bool FWmfMediaStreamSink::ExecuteOnceNoPendingRequests(TFunction<void()>&& ExecuteOnIdle)
+{
+	FScopeLock Lock(&CriticalSection);
+
+	if (NumRequestsInFlight == 0)
+	{
+		return false;
+	}
+
+	if (!OnIdleAction)
+	{
+		OnIdleAction = MoveTemp(ExecuteOnIdle);
+	}
+	return true;
 }
 
 
@@ -686,13 +682,28 @@ STDMETHODIMP FWmfMediaStreamSink::ProcessSample(__RPC__in_opt IMFSample* pSample
 		return E_POINTER;
 	}
 
-	FScopeLock Lock(&CriticalSection);
+	TFunction<void()> TempAction;
 
+	{
+	FScopeLock Lock(&CriticalSection);
+	
 	if (!EventQueue.IsValid())
 	{
 		return MF_E_SHUTDOWN;
 	}
 
+	--NumRequestsInFlight;
+	check(NumRequestsInFlight >= 0);
+	--NumRequestsInFlightWithSystem;
+	check(NumRequestsInFlightWithSystem >= 0);
+
+	bool bIdleWaitInProgress = !!OnIdleAction;
+
+	if (bIdleWaitInProgress && NumRequestsInFlight == 0)
+	{
+		TempAction = MoveTemp(OnIdleAction);
+	}
+
 	if (!CurrentMediaType.IsValid())
 	{
 		UE_LOG(LogWmfMedia, VeryVerbose, TEXT("StreamSink %p: Stream received a sample while not having a valid media type set"), this);
@@ -712,66 +723,41 @@ STDMETHODIMP FWmfMediaStreamSink::ProcessSample(__RPC__in_opt IMFSample* pSample
 	}
 
 	UE_LOG(LogWmfMedia, VeryVerbose, TEXT("StreamSink::ProcessSample Sample time %f"), FTimespan::FromMicroseconds(Time / 10).GetTotalSeconds());
-#if WMFMEDIA_PLAYER_VERSION == 1
-	SampleQueue.Add({ MFSTREAMSINK_MARKER_DEFAULT, NULL, pSample });
-#else // WMFMEDIA_PLAYER_VERSION == 1
-	SampleQueue.Enqueue({ MFSTREAMSINK_MARKER_DEFAULT, NULL, pSample });
-#endif // WMFMEDIA_PLAYER_VERSION == 1
 
-	// finish pre-rolling
-	if (Prerolling)
+	if (!bSeekInProgress)
 	{
-		if (IsVideoSampleQueueFull())
+		if (bScrubbingInProgress)
 		{
-			UE_LOG(LogWmfMedia, VeryVerbose, TEXT("StreamSink %p: Preroll complete, %d samples queued"), this, VideoSampleQueue->Num());
-			Prerolling = false;
-			return QueueEvent(MEStreamSinkPrerolled, GUID_NULL, S_OK, NULL);
-		}
-		else
-		{
-			TComPtr<IMFSample> NextSample;
-			if (GetNextSample(NextSample))
-			{
-				CopyTextureAndEnqueueSample(NextSample);
-				UE_LOG(LogWmfMedia, VeryVerbose, TEXT("StreamSink::ProcessSample Request Sample"));
-				return QueueEvent(MEStreamSinkRequestSample, GUID_NULL, S_OK, NULL);
-			}
-			else
-			{
-				UE_LOG(LogWmfMedia, VeryVerbose, TEXT("StreamSink %p: Preroll complete, %d samples queued"), this, VideoSampleQueue->Num());
-				Prerolling = false;
-				return QueueEvent(MEStreamSinkPrerolled, GUID_NULL, S_OK, NULL);
-			}
-		}
-	}
-	else if (ClockRate == 0.0f)
-	{
-		TComPtr<IMFSample> NextSample;
-		if (GetNextSample(NextSample))
-		{
-			ScheduleWaitForNextSample(NextSample);
+			bScrubbingInProgress = false;
+
+			// Send out scrub complete event.
+			QueueEvent(MEStreamSinkScrubSampleComplete, GUID_NULL, S_OK, NULL);
 		}
 
-#if WMFMEDIA_PLAYER_VERSION == 1
-		return QueueEvent(MEStreamSinkScrubSampleComplete, GUID_NULL, S_OK, NULL);
-#else //  WMFMEDIA_PLAYER_VERSION == 1
-		return S_OK;
-#endif //  WMFMEDIA_PLAYER_VERSION == 1
-	}
-	else
-	{
-		TComPtr<IMFSample> NextSample;
-		if (GetNextSample(NextSample))
+		if (IsVideoSampleQueueFull())
 		{
-			ScheduleWaitForNextSample(NextSample);
+			// This should not happen, but if it does: throw an error (we'll loose the sample)
+			return MF_E_INVALID_STATE_TRANSITION;
 		}
-		else
+
+		// Take sample and put it into the output queue
+		CopyTextureAndEnqueueSample(pSample);
+
+		if (!bIdleWaitInProgress && !IsVideoSampleQueueFull() && ClockRate != 0.0f)
 		{
 			UE_LOG(LogWmfMedia, VeryVerbose, TEXT("StreamSink::ProcessSample Request Sample2"));
 			QueueEvent(MEStreamSinkRequestSample, GUID_NULL, S_OK, NULL);
+			++NumRequestsInFlight;
 		}
-		return S_OK;
 	}
+	}
+
+	if (TempAction)
+	{
+		TempAction();
+	}
+
+	return S_OK;
 }
 
 
@@ -822,9 +808,9 @@ STDMETHODIMP_(ULONG) FWmfMediaStreamSink::Release()
 
 bool FWmfMediaStreamSink::IsVideoSampleQueueFull() const
 {
-	const int32 NumberOfQueueFrames = 3;
+	const int32 NumberOfQueueFrames = kMaxFramesInOutputQueue;
 	const int32 MinNumberOfQueueFrames = FMath::Min(NumberOfQueueFrames, FMediaPlayerQueueDepths::MaxVideoSinkDepth);
-	return (VideoSampleQueue->Num() >= MinNumberOfQueueFrames);
+	return ((VideoSampleQueue->Num() + NumRequestsInFlight)	>= MinNumberOfQueueFrames);
 }
 
 
@@ -1087,46 +1073,6 @@ void FWmfMediaStreamSink::CopyTextureAndEnqueueSample(IMFSample* pSample)
 }
 
 
-STDMETHODIMP FWmfMediaStreamSink::Invoke(IMFAsyncResult* pAsyncResult)
-{
-	UE_LOG(LogWmfMedia, VeryVerbose, TEXT("StreamSink %p:Invoke"), this);
-	FScopeLock Lock(&CriticalSection);
-
-	TComPtr<IMFSample> NextSample;
-	if (GetNextSample(NextSample))
-	{
-		// process next samples
-		ScheduleWaitForNextSample(NextSample);
-		return S_OK;
-	}
-	else
-	{
-		UE_LOG(LogWmfMedia, VeryVerbose, TEXT("StreamSink %p:Invoke Request Sample"), this);
-		return QueueEvent(MEStreamSinkRequestSample, GUID_NULL, S_OK, NULL);
-	}
-}
-
-
-STDMETHODIMP FWmfMediaStreamSink::GetParameters(DWORD* pdwFlags, DWORD* pdwQueue)
-{
-	return E_NOTIMPL;
-}
-
-
-void FWmfMediaStreamSink::SetPresentationClock(IMFPresentationClock* InPresentationClock)
-{
-	FScopeLock Lock(&CriticalSection);
-	PresentationClock = InPresentationClock;
-}
-
-
-void FWmfMediaStreamSink::SetClockRate(float InClockRate)
-{
-	UE_LOG(LogWmfMedia, VeryVerbose, TEXT("StreamSink %p:SetClockRate %f"), this, InClockRate);
-	ClockRate = InClockRate;
-}
-
-
 void FWmfMediaStreamSink::SetMediaSamplePoolAndQueue(
 	TSharedPtr<FWmfMediaHardwareVideoDecodingTextureSamplePool>& InVideoSamplePool,
 	TMediaSampleQueue<IMediaTextureSample>* InVideoSampleQueue,
@@ -1139,136 +1085,19 @@ void FWmfMediaStreamSink::SetMediaSamplePoolAndQueue(
 }
 
 
-void FWmfMediaStreamSink::ScheduleWaitForNextSample(IMFSample* pSample)
+void FWmfMediaStreamSink::SetClockRate(float InClockRate)
 {
-	UE_LOG(LogWmfMedia, VeryVerbose, TEXT("StreamSink::ScheduleWaitForNextSample VideoSampleQueue:%d Rate:%f ThreadId:%x"),
-		VideoSampleQueue->Num(), ClockRate, FPlatformTLS::GetCurrentThreadId());
-
-#if WMFMEDIA_PLAYER_VERSION >= 2
-	double CurrentTime = -1.0f;
-	bool bIsSampleRequested = false;
-	bool bIsThisSampleDesiredSample = false;
-
-	// If we are paused, then check to see if this is the sample we are waiting for.
-	if (ClockRate == 0.0f)
-	{
-		// Get the current time.
-		MFTIME ClockTime;
-		MFTIME SystemTime;
-		if (SUCCEEDED(PresentationClock->GetCorrelatedTime(0, &ClockTime, &SystemTime)))
-		{
-			FTimespan Time = FTimespan(ClockTime);
-			CurrentTime = Time.GetTotalSeconds();
-			
-			// Get the time for this sample.
-			LONGLONG SampleTimeLong = 0;
-			pSample->GetSampleTime(&SampleTimeLong);
-			FTimespan SampleTime = FTimespan::FromMicroseconds(SampleTimeLong / 10);
-
-			// Keep requesting a new sample if this one is before the current time.
-			bIsSampleRequested = SampleTime < Time;
-			
-			// Is this the sample we want?
-			bIsThisSampleDesiredSample = FMath::IsNearlyEqual(SampleTime.GetTotalSeconds(), CurrentTime, (double)KINDA_SMALL_NUMBER);
-			if (bIsThisSampleDesiredSample)
-			{
-				// Send out scrub complete event.
-				UE_LOG(LogWmfMedia, VeryVerbose, TEXT("ScheduleWaitForNextSample send MEStreamSinkScrubSampleComplete."));
-				QueueEvent(MEStreamSinkScrubSampleComplete, GUID_NULL, S_OK, NULL);
-				
-				// This might be true due to floating point issues, so force it to false.
-				bIsSampleRequested = false;
-			}
-
-			UE_LOG(LogWmfMedia, VeryVerbose, TEXT("ScheduleWaitForNextSample Time:%f Sample:%f RequestSample:%d"),
-				Time.GetTotalSeconds(), SampleTime.GetTotalSeconds(), bIsSampleRequested);
-		}
-	}
-#endif // WMFMEDIA_PLAYER_VERSION >= 2
-
-	if (ClockRate == 0.0f)
-	{
-		// Scrubbing, drop all queued samples
-		while (VideoSampleQueue->Num())
-		{
-			TSharedPtr<IMediaTextureSample, ESPMode::ThreadSafe> Sample;
-#if WMFMEDIA_PLAYER_VERSION >= 2
-			// Don't drop the sample if its at the current time.
-			if (VideoSampleQueue->Peek(Sample))
-			{
-				double SampleTime = Sample->GetTime().Time.GetTotalSeconds();
-				if (FMath::IsNearlyEqual(SampleTime, CurrentTime, (double)KINDA_SMALL_NUMBER))
-				{
-					break;
-				}
-			}
-#endif // WMFMEDIA_PLAYER_VERSION >= 2
-			VideoSampleQueue->Dequeue(Sample);
-			UE_LOG(LogWmfMedia, VeryVerbose, TEXT("ScheduleWaitForNextSample drop sample:%f"), Sample.IsValid() ? Sample->GetTime().Time.GetTotalSeconds() : -1.0);
-		}
-	}
-
-	if (IsVideoSampleQueueFull())
-	{
-		// Return sample to internal queue
-#if WMFMEDIA_PLAYER_VERSION == 1
-		SampleQueue.Push({ MFSTREAMSINK_MARKER_DEFAULT, NULL, pSample });
-#else // WMFMEDIA_PLAYER_VERSION == 1
-		TQueue<FQueuedSample> Queue;
-		FQueuedSample TempSample;
-		Queue.Enqueue({MFSTREAMSINK_MARKER_DEFAULT, NULL, pSample});
-		while (SampleQueue.Dequeue(TempSample))
-		{
-			Queue.Enqueue(TempSample);
-		}
-		while (Queue.Dequeue(TempSample))
-		{
-			SampleQueue.Enqueue(TempSample);
-		}
-#endif // WMFMEDIA_PLAYER_VERSION == 1
-	}
-	else
-	{
-		CopyTextureAndEnqueueSample(pSample);
-	}
+	FScopeLock Lock(&CriticalSection);
 
-#if WMFMEDIA_PLAYER_VERSION == 1
-	if (WaitTimer != nullptr && ClockRate != 0.0f)
-#else // WMFMEDIA_PLAYER_VERSION == 1
-	if (WaitTimer != nullptr && ((ClockRate != 0.0f) || bIsSampleRequested))
-#endif // WMFMEDIA_PLAYER_VERSION == 1
-	{
-		// Re-schedule 
-		const LONGLONG OneMilliSeconds = 10000;
-		LARGE_INTEGER llDueTime;
-		llDueTime.QuadPart = -4 * OneMilliSeconds;
-		if (SetWaitableTimer(WaitTimer, &llDueTime, 0, NULL, NULL, FALSE) == 0)
-		{
-			UE_LOG(LogWmfMedia, VeryVerbose, TEXT("SetWaitableTimer Error"));
-			return;
-		}
+	ClockRate = InClockRate;
 
-		TComPtr<IMFAsyncResult> pAsyncResult;
-		HRESULT Result = MFCreateAsyncResult(nullptr, this, nullptr, &pAsyncResult);
-		if (SUCCEEDED(Result))
-		{
-			MFPutWaitingWorkItem(WaitTimer, 0, pAsyncResult, nullptr);
-			UE_LOG(LogWmfMedia, VeryVerbose, TEXT("MFPutWaitingWorkItem"));
-		}
-		else
-		{
-			UE_LOG(LogWmfMedia, VeryVerbose, TEXT("MFPutWaitingWorkItem Error"));
-			return;
-		}
-	}
-	else
+	if (ClockRate != 0.0f)
 	{
-		UE_LOG(LogWmfMedia, VeryVerbose, TEXT("WaitTimer == 0"));
+		bScrubbingInProgress = false;
+		bSystemIsScrubbing = false;
 	}
 
-	UE_LOG(LogWmfMedia, VeryVerbose, TEXT("ScheduleWaitForNextSample End"));
-
-	return;
+	QueueEvent(MEStreamSinkRateChanged, GUID_NULL, S_OK, NULL);
 }
 
 #include "Windows/HideWindowsPlatformTypes.h"
diff --git a/Engine/Plugins/Media/WmfMedia/Source/WmfMedia/Private/Wmf/WmfMediaStreamSink.h b/Engine/Plugins/Media/WmfMedia/Source/WmfMedia/Private/Wmf/WmfMediaStreamSink.h
--- a/Engine/Plugins/Media/WmfMedia/Source/WmfMedia/Private/Wmf/WmfMediaStreamSink.h
+++ b/Engine/Plugins/Media/WmfMedia/Source/WmfMedia/Private/Wmf/WmfMediaStreamSink.h
@@ -30,7 +30,6 @@ class FWmfMediaStreamSink
 	: public IMFGetService
 	, public IMFMediaTypeHandler
 	, public IMFStreamSink
-	, public IMFAsyncCallback
 {
 public:
 
@@ -55,14 +54,6 @@ public:
 
 public:
 
-	/**
-	 * Get the next sample in the queue.
-	 *
-	 * @param OutSample Will contain the sample.
-	 * @return true if a sample was returned, false if the queue is empty.
-	 */
-	bool GetNextSample(TComPtr<IMFSample>& OutSample);
-
 	/**
 	 * Initialize this sink.
 	 *
@@ -89,16 +80,6 @@ public:
 	 */
 	HRESULT Pause();
 
-	/**
-	 * Preroll the sink.
-	 *
-	 * This method is called by the owner media sink.
-	 *
-	 * @return Result code.
-	 * @see Pause, Restart, Start, Stop
-	 */
-	HRESULT Preroll();
-
 	/**
 	 * Restart the stream.
 	 *
@@ -173,12 +154,6 @@ public:
 	STDMETHODIMP PlaceMarker(MFSTREAMSINK_MARKER_TYPE eMarkerType, __RPC__in const PROPVARIANT* pvarMarkerValue, __RPC__in const PROPVARIANT* pvarContextValue);
 	STDMETHODIMP ProcessSample(__RPC__in_opt IMFSample* pSample);
 
-public:
-
-	//~ IMFAsyncCallback interface
-	STDMETHODIMP Invoke(IMFAsyncResult* pAsyncResult);
-	STDMETHODIMP GetParameters(DWORD* pdwFlags, DWORD* pdwQueue);
-
 public:
 
 	//~ IUnknown interface
@@ -187,13 +162,6 @@ public:
 	STDMETHODIMP QueryInterface(REFIID RefID, void** Object);
 	STDMETHODIMP_(ULONG) Release();
 
-	/**
-	 * Set current presentation 
-	 *
-	 * @param InPresentationClock Presentation clock to use
-	 */
-	void SetPresentationClock(IMFPresentationClock* InPresentationClock);
-
 	/**
 	 * Set current clock rate
 	 *
@@ -210,6 +178,24 @@ public:
 	 */
 	void SetMediaSamplePoolAndQueue(TSharedPtr<FWmfMediaHardwareVideoDecodingTextureSamplePool>& InVideoSamplePool, TMediaSampleQueue<IMediaTextureSample>* InVideoSampleQueue, TFunction<FMediaTimeStamp(FTimespan, EMediaTrackType)>&& InAdjustTimeStamp);
 
+	/**
+	 * Request more data being delivered to video output queue 
+	 */
+	void OnRequestMoreData();
+
+	/**
+	 * Notify instance that a seek logically started
+	 */
+	void OnSeekStarted();
+
+
+	void OnSessionEnded();
+
+	/**
+	 * Register lambda to execute (once) as soon as any pending requests are done (if any are pending)
+	 */
+	bool ExecuteOnceNoPendingRequests(TFunction<void()>&& ExecuteOnIdle);
+
 private:
 
 	/** Hidden destructor (this class is reference counted). */
@@ -222,19 +208,6 @@ private:
 	 */
 	void CopyTextureAndEnqueueSample(IMFSample* pSample);
 
-	/**
-	 * Check whether sample is ready to display otherwise re-schedule
-	 *
-	 * @param pSample Sample from decoder
-	 * @note This function is not thread-safe but all caller are using a Scope Lock
-	 */
-	void ScheduleWaitForNextSample(IMFSample* pSample);
-
-	/**
-	 * Close Timer
-	 */
-	void CloseTimer();
-
 	/**
 	 * Check if video sample queue is full
 	 * @return true when Video Sample Queue is full.
@@ -258,9 +231,6 @@ private:
 	/** The decoder we belong to. */
 	TComPtr<WmfMediaDecoder> Decoder;
 
-	/** Whether the sink is currently prerolling samples. */
-	bool Prerolling;
-
 	/** Holds a reference counter for this instance. */
 	int32 RefCount;
 
@@ -270,9 +240,6 @@ private:
 	/** The sink's major media type. */
 	const GUID StreamType;
 
-	/** Presentation clock which dictate when to present a sample. */
-	TComPtr<IMFPresentationClock> PresentationClock;
-
 	/** ClockRate to support reverse playback */
 	float ClockRate;
 
@@ -284,6 +251,23 @@ private:
 
 	/** Video sample queue which is filled with current sample */
 	TMediaSampleQueue<IMediaTextureSample>* VideoSampleQueue;
+
+	/** Number of requests for new video samples currently in flight */
+	int32 NumRequestsInFlight;
+	int32 NumRequestsInFlightWithSystem;
+	TFunction<void()> OnIdleAction;
+
+	/** If true seek is in progress logically */
+	bool bSeekInProgress;
+
+	/** If true scrubbing is in progress logically */
+	bool bScrubbingInProgress;
+
+	/** If true system is in scrubbing mode (play with rate==0) */
+	bool bSystemIsScrubbing;
+
+	/** If true segment just ended, but is still logically being processed */
+	bool bSegmentEnded;
 	
 	/** Callback to adjust given timestamp for the output queue */
 	TFunction<FMediaTimeStamp(FTimespan,EMediaTrackType)> AdjustTimeStamp;
@@ -303,13 +287,6 @@ private:
 		/** The media sample. */
 		TComPtr<IMFSample> Sample;
 	};
-
-	/** Media sample array used as dequeue to store extra samples when video sample queue is full */
-#if WMFMEDIA_PLAYER_VERSION == 1
-	TArray<FQueuedSample> SampleQueue;
-#else // WMFMEDIA_PLAYER_VERSION == 1
-	TQueue<FQueuedSample> SampleQueue;
-#endif // WMFMEDIA_PLAYER_VERSION == 1
 };
 
 #endif
diff --git a/Engine/Plugins/Media/WmfMedia/Source/WmfMedia/WmfMedia.Build.cs b/Engine/Plugins/Media/WmfMedia/Source/WmfMedia/WmfMedia.Build.cs
--- a/Engine/Plugins/Media/WmfMedia/Source/WmfMedia/WmfMedia.Build.cs
+++ b/Engine/Plugins/Media/WmfMedia/Source/WmfMedia/WmfMedia.Build.cs
@@ -51,8 +51,6 @@ namespace UnrealBuildTool.Rules
 				PublicDelayLoadDLLs.Add("mfplay.dll");
 				PublicDelayLoadDLLs.Add("shlwapi.dll");
 			}
-			
-			PublicDefinitions.Add("WMFMEDIA_PLAYER_VERSION=2");
 		}
 	}
 }
