WmfMedia: Fixed the 'reachedend' event being sent twice

#rnx

diff --git a/Engine/Plugins/Media/WmfMedia/Source/WmfMedia/Private/Wmf/WmfMediaSession.cpp b/Engine/Plugins/Media/WmfMedia/Source/WmfMedia/Private/Wmf/WmfMediaSession.cpp
--- a/Engine/Plugins/Media/WmfMedia/Source/WmfMedia/Private/Wmf/WmfMediaSession.cpp
+++ b/Engine/Plugins/Media/WmfMedia/Source/WmfMedia/Private/Wmf/WmfMediaSession.cpp
@@ -87,9 +87,19 @@ FWmfMediaSession::~FWmfMediaSession()
 
 void FWmfMediaSession::GetEvents(TArray<EMediaEvent>& OutEvents)
 {
-	// Are we waiting for the end?
+	// Are we waiting for the end (of non-looping playback)?
 	if (bIsWaitingForEnd)
 	{
+		/*
+			Realistically we cannot check the tracks to have delivered all samples, because
+			tracks may have different durations. If audio is shorter then the facade will not
+			pull the remaining video samples. Also, WmfMedia delivers samples outside the movie's
+			duration (with mp4/mov the duration is reported as specified in the movie box, but
+			the tracks and the delivered samples may be longer, so we end up getting extra
+			samples at the end that are outside the duration and will not be consumed, thus
+			sticking around in the sample queues that will not empty.
+		*/
+#if 0
 		// Is Tracks done with all its samples?
 		TSharedPtr<FWmfMediaTracks, ESPMode::ThreadSafe> TracksPinned = Tracks.Pin();
 		if (TracksPinned.IsValid())
@@ -112,6 +122,10 @@ void FWmfMediaSession::GetEvents(TArray<EMediaEvent>& OutEvents)
 		{
 			DeferredEvents.Enqueue(EMediaEvent::PlaybackEndReached);
 		}
+#else
+		bIsWaitingForEnd = false;
+		DeferredEvents.Enqueue(EMediaEvent::PlaybackEndReached);
+#endif
 	}
 
 	// Hand out the events...
@@ -702,7 +716,7 @@ STDMETHODIMP FWmfMediaSession::Invoke(IMFAsyncResult* AsyncResult)
 		break;
 
 	case MEError:
-		HandleError(EventStatus);	
+		HandleError(EventStatus);
 		break;
 
 	case MEReconnectEnd:
@@ -807,7 +821,7 @@ STDMETHODIMP FWmfMediaSession::Invoke(IMFAsyncResult* AsyncResult)
 	UE_LOG(LogWmfMedia, VeryVerbose, TEXT("Session %p: WaitingForEnd:%d"),
 		this,
 		bIsWaitingForEnd);
-	
+
 	return S_OK;
 }
 
@@ -836,7 +850,7 @@ STDMETHODIMP FWmfMediaSession::QueryInterface(REFIID RefID, void** Object)
 STDMETHODIMP_(ULONG) FWmfMediaSession::Release()
 {
 	int32 CurrentRefCount = FPlatformAtomics::InterlockedDecrement(&RefCount);
-	
+
 	if (CurrentRefCount == 0)
 	{
 		delete this;
@@ -866,7 +880,7 @@ FWmfMediaSession::EWorkItemResult FWmfMediaSession::CommitRate(float Rate, FRate
 	{
 		/*
 		* No rate control
-		* 
+		*
 		* Session is only able to support pause & play
 		* (no reverse, scrubbing or speed changes)
 		*/
@@ -932,7 +946,7 @@ FWmfMediaSession::EWorkItemResult FWmfMediaSession::CommitRate(float Rate, FRate
 	// ---------------------------------------------------------------------------------------------------------------
 	/**
 	 * Rate control is present
-	 * 
+	 *
 	 * Session can plat at variable speeds, reverse (possibly) and scrub (possibly)
 	 */
 
@@ -945,7 +959,7 @@ FWmfMediaSession::EWorkItemResult FWmfMediaSession::CommitRate(float Rate, FRate
 		{
 			/**
 			 * Reversal of playback direction
-			 * 
+			 *
 			 * System needs to be stopped and we ensure no pending requests as WMF tends to get confused if they come in while we do this
 			 */
 
@@ -953,7 +967,7 @@ FWmfMediaSession::EWorkItemResult FWmfMediaSession::CommitRate(float Rate, FRate
 			if (SessionState != EMediaState::Stopped)
 			{
 				// No...
-				
+
 				// If we reverse, we might loose issued sample requests (without us knowing) - wait until all current ones are done
 				TSharedPtr<FWmfMediaTracks, ESPMode::ThreadSafe> TracksPinned = Tracks.Pin();
 				if (TracksPinned.IsValid())
@@ -1044,7 +1058,7 @@ FWmfMediaSession::EWorkItemResult FWmfMediaSession::CommitRate(float Rate, FRate
 								{
 									State->LastTime = FTimespan(ClockTime);
 								}
-								
+
 								// Stop the session
 								const HRESULT Result = MediaSession->Stop();
 								if (FAILED(Result))
@@ -1142,7 +1156,7 @@ FWmfMediaSession::EWorkItemResult FWmfMediaSession::CommitRate(float Rate, FRate
 						{
 							State->LastTime = FTimespan(ClockTime);
 						}
- 
+
 						// Stop the session
 						const HRESULT Result = MediaSession->Stop();
 						if (FAILED(Result))
@@ -1603,11 +1617,6 @@ void FWmfMediaSession::HandleSessionEnded()
 	// We only loop if we don't have a seek in progress right now
 	if (!bSeekActive)
 	{
-		if (!ShouldLoop)
-		{
-			DeferredEvents.Enqueue(EMediaEvent::PlaybackEndReached);
-		}
-
 		if (ShouldLoop)
 		{
 			// loop back to beginning/end
