Media: v2 timing enabled media players now receive the sequence index the facade is expecting on a Seek() via parameter instead of having to replicate the facade's internal behavior which is very prone to errors.

#rnx

diff --git a/Engine/Plugins/Media/ElectraPlayer/Source/ElectraPlayerPlugin/Private/ElectraPlayerPlugin.cpp b/Engine/Plugins/Media/ElectraPlayer/Source/ElectraPlayerPlugin/Private/ElectraPlayerPlugin.cpp
--- a/Engine/Plugins/Media/ElectraPlayer/Source/ElectraPlayerPlugin/Private/ElectraPlayerPlugin.cpp
+++ b/Engine/Plugins/Media/ElectraPlayer/Source/ElectraPlayerPlugin/Private/ElectraPlayerPlugin.cpp
@@ -941,11 +941,13 @@ bool FElectraPlayerPlugin::SetRate(float Rate)
 }
 
 
-bool FElectraPlayerPlugin::Seek(const FTimespan& Time)
+bool FElectraPlayerPlugin::Seek(const FTimespan& Time, const FMediaSeekParams& InAdditionalParams)
 {
 	UE_LOG(LogElectraPlayerPlugin, Verbose, TEXT("[%u] IMediaControls::Seek() to %s"), PlayerUniqueID, *Time.ToString(TEXT("%h:%m:%s.%f")));
 	CSV_EVENT(ElectraPlayer, TEXT("Seeking"));
-	return Player->Seek(Time);
+	IElectraPlayerInterface::FSeekParam sp;
+	sp.SequenceIndex = InAdditionalParams.NewSequenceIndex;
+	return Player->Seek(Time, sp);
 }
 
 TRange<FTimespan> FElectraPlayerPlugin::GetPlaybackTimeRange(EMediaTimeRangeType InRangeToGet) const
diff --git a/Engine/Plugins/Media/ElectraPlayer/Source/ElectraPlayerPlugin/Private/ElectraPlayerPlugin.h b/Engine/Plugins/Media/ElectraPlayer/Source/ElectraPlayerPlugin/Private/ElectraPlayerPlugin.h
--- a/Engine/Plugins/Media/ElectraPlayer/Source/ElectraPlayerPlugin/Private/ElectraPlayerPlugin.h
+++ b/Engine/Plugins/Media/ElectraPlayer/Source/ElectraPlayerPlugin/Private/ElectraPlayerPlugin.h
@@ -131,7 +131,9 @@ private:
 	FTimespan GetTime() const override;
 	float GetRate() const override;
 	bool SetRate(float Rate) override;
-	bool Seek(const FTimespan& Time) override;
+	bool Seek(const FTimespan& Time) override
+	{ check(!"You have to call the override with additional options!"); return false; }
+	bool Seek(const FTimespan& InNewTime, const FMediaSeekParams& InAdditionalParams) override;
 
 	TRange<FTimespan> GetPlaybackTimeRange(EMediaTimeRangeType InRangeToGet) const override;
 	bool SetPlaybackTimeRange(const TRange<FTimespan>& InTimeRange) override;
diff --git a/Engine/Plugins/Media/ElectraPlayer/Source/ElectraPlayerRuntime/Private/ElectraPlayer.cpp b/Engine/Plugins/Media/ElectraPlayer/Source/ElectraPlayerRuntime/Private/ElectraPlayer.cpp
--- a/Engine/Plugins/Media/ElectraPlayer/Source/ElectraPlayerRuntime/Private/ElectraPlayer.cpp
+++ b/Engine/Plugins/Media/ElectraPlayer/Source/ElectraPlayerRuntime/Private/ElectraPlayer.cpp
@@ -199,7 +199,6 @@ void FElectraPlayer::ClearToDefaultState()
 {
 	FScopeLock lock(&PlayerLock);
 
-	CurrentSequenceIndex = 0;
 	PlayerState.Reset();
 	NumTracksAudio = 0;
 	NumTracksVideo = 0;
@@ -1247,15 +1246,8 @@ void FElectraPlayer::CalculateTargetSeekTime(FTimespan& OutTargetTime, const FTi
 }
 
 
-bool FElectraPlayer::Seek(const FTimespan& Time)
-{
-	FSeekParam NoParams;
-	return Seek(Time, NoParams);
-}
-
 bool FElectraPlayer::Seek(const FTimespan& Time, const FSeekParam& Param)
 {
-	++CurrentSequenceIndex;
 	if (CurrentPlayer.Get())
 	{
 		FTimespan Target;
@@ -1265,13 +1257,9 @@ bool FElectraPlayer::Seek(const FTimespan& Time, const FSeekParam& Param)
 		{
 			seek.Time.SetFromTimespan(Target);
 		}
+		check(Param.SequenceIndex.IsSet());
 		seek.NewSequenceIndex = Param.SequenceIndex;
 		seek.StartingBitrate = Param.StartingBitrate;
-		seek.bOptimizeForScrubbing = Param.bOptimizeForScrubbing;
-		if (!seek.NewSequenceIndex.IsSet())
-		{
-			seek.NewSequenceIndex = CurrentSequenceIndex;
-		}
 		bInitialSeekPerformed = true;
 		bDiscardOutputUntilCleanStart = true;
 		CurrentPlayer->AdaptivePlayer->SeekTo(seek);
diff --git a/Engine/Plugins/Media/ElectraPlayer/Source/ElectraPlayerRuntime/Private/ElectraPlayer.h b/Engine/Plugins/Media/ElectraPlayer/Source/ElectraPlayerRuntime/Private/ElectraPlayer.h
--- a/Engine/Plugins/Media/ElectraPlayer/Source/ElectraPlayerRuntime/Private/ElectraPlayer.h
+++ b/Engine/Plugins/Media/ElectraPlayer/Source/ElectraPlayerRuntime/Private/ElectraPlayer.h
@@ -116,7 +116,6 @@ public:
 	float GetRate() const override;
 	bool SetRate(float Rate) override;
 
-	bool Seek(const FTimespan& Time) override;
 	bool Seek(const FTimespan& Time, const FSeekParam& Param) override;
 	void SetFrameAccurateSeekMode(bool bEnableFrameAccuracy) override;
 
@@ -407,9 +406,6 @@ private:
 	// Delegate to talk back to adapter host
 	TWeakPtr<IElectraPlayerAdapterDelegate, ESPMode::ThreadSafe>	AdapterDelegate;
 
-	// The current sequence index as expected by the media player facade.
-	int32 CurrentSequenceIndex = 0;
-
 	// Contains number of audio tracks available to expose it later.
 	int32											NumTracksAudio;
 	int32											NumTracksVideo;
diff --git a/Engine/Plugins/Media/ElectraPlayer/Source/ElectraPlayerRuntime/Public/IElectraPlayerInterface.h b/Engine/Plugins/Media/ElectraPlayer/Source/ElectraPlayerRuntime/Public/IElectraPlayerInterface.h
--- a/Engine/Plugins/Media/ElectraPlayer/Source/ElectraPlayerRuntime/Public/IElectraPlayerInterface.h
+++ b/Engine/Plugins/Media/ElectraPlayer/Source/ElectraPlayerRuntime/Public/IElectraPlayerInterface.h
@@ -264,10 +264,8 @@ public:
 	{
 		TOptional<int32> SequenceIndex;
 		TOptional<int32> StartingBitrate;
-		TOptional<bool> bOptimizeForScrubbing;
 	};
 
-	virtual bool Seek(const FTimespan& Time) = 0;
 	virtual bool Seek(const FTimespan& Time, const FSeekParam& Param) = 0;
 	virtual void SetFrameAccurateSeekMode(bool bEnableFrameAccuracy) = 0;
 
diff --git a/Engine/Plugins/Media/ElectraPlayer/Source/ElectraProtron/Private/Player/ElectraProtronPlayer.cpp b/Engine/Plugins/Media/ElectraPlayer/Source/ElectraProtron/Private/Player/ElectraProtronPlayer.cpp
--- a/Engine/Plugins/Media/ElectraPlayer/Source/ElectraProtron/Private/Player/ElectraProtronPlayer.cpp
+++ b/Engine/Plugins/Media/ElectraPlayer/Source/ElectraProtron/Private/Player/ElectraProtronPlayer.cpp
@@ -349,7 +349,7 @@ bool FElectraProtronPlayer::IsLooping() const
 	return Player.IsValid() ? Player->IsLooping() : false;
 }
 
-bool FElectraProtronPlayer::Seek(const FTimespan& InTime)
+bool FElectraProtronPlayer::Seek(const FTimespan& InTime, const FMediaSeekParams& InAdditionalParams)
 {
 	// Create a new sample queue into which the new decoded samples will be added.
 	TSharedPtr<FMediaSamples, ESPMode::ThreadSafe> SampleQueue = MakeShared<FMediaSamples, ESPMode::ThreadSafe>();
@@ -362,9 +362,9 @@ bool FElectraProtronPlayer::Seek(const FTimespan& InTime)
 	{
 		return false;
 	}
-	const int32 SeqChange = Player->GetRate() >= 0.0f ? 1 : -1;
+	check(InAdditionalParams.NewSequenceIndex.IsSet());
 	int32 SequenceIndexNow = CurrentSequenceIndex;
-	CurrentSequenceIndex += SeqChange;
+	CurrentSequenceIndex = InAdditionalParams.NewSequenceIndex.GetValue();
 	bool bOk = Player->Seek(InTime, CurrentSampleQueue, CurrentSequenceIndex);
 	if (bOk)
 	{
diff --git a/Engine/Plugins/Media/ElectraPlayer/Source/ElectraProtron/Private/Player/ElectraProtronPlayer.h b/Engine/Plugins/Media/ElectraPlayer/Source/ElectraProtron/Private/Player/ElectraProtronPlayer.h
--- a/Engine/Plugins/Media/ElectraPlayer/Source/ElectraProtron/Private/Player/ElectraProtronPlayer.h
+++ b/Engine/Plugins/Media/ElectraPlayer/Source/ElectraProtron/Private/Player/ElectraProtronPlayer.h
@@ -77,7 +77,9 @@ protected:
 	TRangeSet<float> GetSupportedRates(EMediaRateThinning InThinning) const override;
 	FTimespan GetTime() const override;
 	bool IsLooping() const override;
-	bool Seek(const FTimespan& InTime) override;
+	bool Seek(const FTimespan& Time) override
+	{ check(!"You have to call the override with additional options!"); return false; }
+	bool Seek(const FTimespan& InNewTime, const FMediaSeekParams& InAdditionalParams) override;
 	bool SetLooping(bool bInLooping) override;
 	bool SetRate(float InRate) override;
 	bool SetNativeVolume(float InVolume) override;
diff --git a/Engine/Plugins/Media/ImgMedia/Source/ImgMedia/Internal/Player/ImgMediaPlayer.h b/Engine/Plugins/Media/ImgMedia/Source/ImgMedia/Internal/Player/ImgMediaPlayer.h
--- a/Engine/Plugins/Media/ImgMedia/Source/ImgMedia/Internal/Player/ImgMediaPlayer.h
+++ b/Engine/Plugins/Media/ImgMedia/Source/ImgMedia/Internal/Player/ImgMediaPlayer.h
@@ -100,7 +100,9 @@ protected:
 	virtual TRangeSet<float> GetSupportedRates(EMediaRateThinning Thinning) const override;
 	virtual FTimespan GetTime() const override;
 	virtual bool IsLooping() const override;
-	virtual bool Seek(const FTimespan& Time) override;
+	virtual bool Seek(const FTimespan& Time) override
+	{ check(!"You have to call the override with additional options!"); return false; }
+	virtual bool Seek(const FTimespan& InNewTime, const FMediaSeekParams& InAdditionalParams) override;
 	virtual bool SetLooping(bool Looping) override;
 	virtual bool SetRate(float Rate) override;
 	virtual void SetBlockingPlaybackHint(bool bFacadeWillUseBlockingPlayback) override;
diff --git a/Engine/Plugins/Media/ImgMedia/Source/ImgMedia/Private/Player/ImgMediaPlayer.cpp b/Engine/Plugins/Media/ImgMedia/Source/ImgMedia/Private/Player/ImgMediaPlayer.cpp
--- a/Engine/Plugins/Media/ImgMedia/Source/ImgMedia/Private/Player/ImgMediaPlayer.cpp
+++ b/Engine/Plugins/Media/ImgMedia/Source/ImgMedia/Private/Player/ImgMediaPlayer.cpp
@@ -215,7 +215,7 @@ bool FImgMediaPlayer::Open(const FString& Url, const IMediaOptions* Options, con
 
 	// determine image sequence proxy, if any
 	FString Proxy;
-	
+
 	if (Options != nullptr)
 	{
 		Proxy = Options->GetMediaOption(ImgMedia::ProxyOverrideOption, FString());
@@ -377,7 +377,7 @@ bool FImgMediaPlayer::GetPlayerFeatureFlag(EFeatureFlag flag) const
 	default:
 		break;
 	}
-	
+
 	return IMediaPlayer::GetPlayerFeatureFlag(flag);
 }
 
@@ -508,7 +508,7 @@ bool FImgMediaPlayer::IsLooping() const
 }
 
 
-bool FImgMediaPlayer::Seek(const FTimespan& Time)
+bool FImgMediaPlayer::Seek(const FTimespan& Time, const FMediaSeekParams& InAdditionalParams)
 {
 	// validate seek
 	if (!IsInitialized())
@@ -535,7 +535,8 @@ bool FImgMediaPlayer::Seek(const FTimespan& Time)
 	}
 
 	CurrentTime = Time;
-	CurrentSeekIndex += (CurrentRate >= 0.0f) ? 1 : -1;
+	check(InAdditionalParams.NewSequenceIndex.IsSet());
+	CurrentSeekIndex = InAdditionalParams.NewSequenceIndex.GetValue();
 	if (Loader.IsValid())
 	{
 		Loader->Seek(FMediaTimeStamp(CurrentTime, FMediaTimeStamp::MakeSequenceIndex(CurrentSeekIndex, 0)), LastNonZeroRate < 0.0f);
diff --git a/Engine/Plugins/Media/WmfMedia/Source/WmfMedia/Private/Player/WmfMediaTracks.cpp b/Engine/Plugins/Media/WmfMedia/Source/WmfMedia/Private/Player/WmfMediaTracks.cpp
--- a/Engine/Plugins/Media/WmfMedia/Source/WmfMedia/Private/Player/WmfMediaTracks.cpp
+++ b/Engine/Plugins/Media/WmfMedia/Source/WmfMedia/Private/Player/WmfMediaTracks.cpp
@@ -383,20 +383,13 @@ void FWmfMediaTracks::SetSessionState(EMediaState InState)
 }
 
 
-void FWmfMediaTracks::SeekStarted(const FTimespan& InTime, uint32 UserIssuedSeeks, float UnpausedSessionRate)
+void FWmfMediaTracks::SeekStarted(const FTimespan& InTime, int32 InNewSequenceIndex)
 {
 	UE_LOG(LogWmfMedia, VeryVerbose, TEXT("FWmfMediaTracks::SeekStarted %f"), InTime.GetTotalSeconds());
 
 	{
 	FScopeLock Lock(&CriticalSection);
-	if (UnpausedSessionRate >= 0.0f)
-	{
-		SeekIndex += UserIssuedSeeks;
-	}
-	else
-	{
-		SeekIndex -= UserIssuedSeeks;
-	}
+	SeekIndex = InNewSequenceIndex;
 	LoopIndex = 0;
 	}
 
diff --git a/Engine/Plugins/Media/WmfMedia/Source/WmfMedia/Private/Player/WmfMediaTracks.h b/Engine/Plugins/Media/WmfMedia/Source/WmfMedia/Private/Player/WmfMediaTracks.h
--- a/Engine/Plugins/Media/WmfMedia/Source/WmfMedia/Private/Player/WmfMediaTracks.h
+++ b/Engine/Plugins/Media/WmfMedia/Source/WmfMedia/Private/Player/WmfMediaTracks.h
@@ -194,7 +194,7 @@ public:
 	/**
 	 * Notify track of triggered seek
 	 */
-	void SeekStarted(const FTimespan& InTime, uint32 UserIssuedSeeks, float UnpausedSessionRate);
+	void SeekStarted(const FTimespan& InTime, int32 InNewSequenceIndex);
 
 	/**
 	 * Notify track of triggered loop
diff --git a/Engine/Plugins/Media/WmfMedia/Source/WmfMedia/Private/Wmf/WmfMediaSession.cpp b/Engine/Plugins/Media/WmfMedia/Source/WmfMedia/Private/Wmf/WmfMediaSession.cpp
--- a/Engine/Plugins/Media/WmfMedia/Source/WmfMedia/Private/Wmf/WmfMediaSession.cpp
+++ b/Engine/Plugins/Media/WmfMedia/Source/WmfMedia/Private/Wmf/WmfMediaSession.cpp
@@ -60,7 +60,6 @@ FWmfMediaSession::FWmfMediaSession()
 	, Status(EMediaStatus::None)
 	, bShuttingDown(false)
 	, bIsWaitingForEnd(false)
-	, UserIssuedSeeks(0)
 	, bSeekActive(false)
 	, NumSeeksInWorkQueue(0)
 {
@@ -536,14 +535,15 @@ bool FWmfMediaSession::IsLooping() const
 }
 
 
-bool FWmfMediaSession::Seek(const FTimespan& Time)
+bool FWmfMediaSession::Seek(const FTimespan& Time, const FMediaSeekParams& InAdditionalParams)
 {
 	if (MediaSession == nullptr)
 	{
 		return false;
 	}
 
-	UE_LOG(LogWmfMedia, Verbose, TEXT("Session %p: Seeking to %s"), this, *Time.ToString());
+	check(InAdditionalParams.NewSequenceIndex.IsSet());
+	UE_LOG(LogWmfMedia, Verbose, TEXT("Session %p: Seeking to %s with sequence index"), this, *Time.ToString(), InAdditionalParams.NewSequenceIndex.GetValue());
 
 	// validate seek
 	if (!CanControl(EMediaControl::Seek))
@@ -567,6 +567,7 @@ bool FWmfMediaSession::Seek(const FTimespan& Time)
 	}
 
 	++NumSeeksInWorkQueue;
+	NextSequenceIndex = InAdditionalParams.NewSequenceIndex;
 	EnqueueWorkItem(FWorkItem([this, Time](TComPtr<IMFMediaEvent> MediaEvent, MediaEventType EventType, HRESULT EventStatus, FBaseWorkItemState* BaseState) -> EWorkItemResult
 		{
 			auto State = static_cast<FTimeWorkItemState*>(BaseState);
@@ -576,11 +577,7 @@ bool FWmfMediaSession::Seek(const FTimespan& Time)
 			* indeed start the session and interrupt the end of session events (including keeping requests alive etc.)
 			*/
 
-			if (State->State == FTimeWorkItemState::EState::Begin)
-			{
-				++UserIssuedSeeks;
-			}
-			else if (State->State == FTimeWorkItemState::EState::WaitDone)
+			if (State->State == FTimeWorkItemState::EState::WaitDone)
 			{
 				return (EventType == MESessionStarted) ? EWorkItemResult::Done : EWorkItemResult::RetryLater;
 			}
@@ -1315,8 +1312,9 @@ bool FWmfMediaSession::CommitTime(FTimespan Time, bool bIsSeek)
 		TSharedPtr<FWmfMediaTracks, ESPMode::ThreadSafe> TracksPinned = Tracks.Pin();
 		if (TracksPinned.IsValid())
 		{
-			TracksPinned->SeekStarted(OriginalTime, UserIssuedSeeks, UnpausedSessionRate);
-			UserIssuedSeeks = 0;
+			check(NextSequenceIndex.IsSet());
+			TracksPinned->SeekStarted(OriginalTime, NextSequenceIndex.GetValue());
+			NextSequenceIndex.Reset();
 		}
 		bSeekActive = true;
 	}
@@ -1565,7 +1563,7 @@ void FWmfMediaSession::SetSessionRate(float Rate)
 void::FWmfMediaSession::Flush()
 {
 	FScopeLock Lock(&CriticalSection);
-	UserIssuedSeeks = 0;
+	NextSequenceIndex.Reset();
 }
 
 
diff --git a/Engine/Plugins/Media/WmfMedia/Source/WmfMedia/Private/Wmf/WmfMediaSession.h b/Engine/Plugins/Media/WmfMedia/Source/WmfMedia/Private/Wmf/WmfMediaSession.h
--- a/Engine/Plugins/Media/WmfMedia/Source/WmfMedia/Private/Wmf/WmfMediaSession.h
+++ b/Engine/Plugins/Media/WmfMedia/Source/WmfMedia/Private/Wmf/WmfMediaSession.h
@@ -107,7 +107,9 @@ public:
 	virtual TRangeSet<float> GetSupportedRates(EMediaRateThinning Thinning) const override;
 	virtual FTimespan GetTime() const override;
 	virtual bool IsLooping() const override;
-	virtual bool Seek(const FTimespan& Time) override;
+	virtual bool Seek(const FTimespan& Time) override
+	{ check(!"You have to call the override with additional options!"); return false; }
+	virtual bool Seek(const FTimespan& InNewTime, const FMediaSeekParams& InAdditionalParams) override;
 	virtual bool SetLooping(bool Looping) override;
 	virtual bool SetRate(float Rate) override;
 
@@ -363,8 +365,8 @@ private:
 	/** True if we are waiting for the tracks to send out its last samples before we "end". */
 	bool bIsWaitingForEnd;
 
-	/** Number of user initiated seek calls accumulated for next actual seek executed */
-	uint32 UserIssuedSeeks;
+	/** Next requested sequence index on a Seek() request. */
+	TOptional<int32> NextSequenceIndex;
 
 	/** True if seek is logically in progress */
 	bool bSeekActive;
diff --git a/Engine/Source/Runtime/Media/Public/IMediaControls.h b/Engine/Source/Runtime/Media/Public/IMediaControls.h
--- a/Engine/Source/Runtime/Media/Public/IMediaControls.h
+++ b/Engine/Source/Runtime/Media/Public/IMediaControls.h
@@ -6,6 +6,7 @@
 #include "Math/RangeSet.h"
 #include "Misc/EnumClassFlags.h"
 #include "Misc/Timespan.h"
+#include "Misc/Optional.h"
 
 
 /**
@@ -122,6 +123,16 @@ enum class EMediaTimeRangeType
 };
 
 
+/**
+ * Additional parameters passed to Seek().
+ */
+struct FMediaSeekParams
+{
+	/** If set, the new sequence index a v2 timing enabled media player must be using now. */
+	TOptional<int32> NewSequenceIndex;
+};
+
+
 /**
  * Interface for controlling media playback.
  *
@@ -236,6 +247,16 @@ public:
 	{
 	}
 
+
+	/**
+	 * Same as Seek(), but with additional options.
+	 * This default implementation just calls Seek() and ignores options.
+	 */
+	virtual bool Seek(const FTimespan& InNewTime, const FMediaSeekParams& InAdditionalParams)
+	{
+		return Seek(InNewTime);
+	}
+
 public:
 
 	/**
@@ -272,6 +293,7 @@ public:
 	 * @return true on success, false otherwise.
 	 * @see GetDuration, GetTime
 	 */
+	UE_DEPRECATED(5.6, "Use absolute seeking only. This method will be removed soon.")
 	bool Seek(const FTimespan& TimeOffset, EMediaSeekDirection Direction)
 	{
 		FTimespan SeekTime;
diff --git a/Engine/Source/Runtime/MediaUtils/Private/MediaPlayerFacade.cpp b/Engine/Source/Runtime/MediaUtils/Private/MediaPlayerFacade.cpp
--- a/Engine/Source/Runtime/MediaUtils/Private/MediaPlayerFacade.cpp
+++ b/Engine/Source/Runtime/MediaUtils/Private/MediaPlayerFacade.cpp
@@ -1215,8 +1215,23 @@ bool FMediaPlayerFacade::Seek(const FTimespan& InTime)
 		Time = InTime;
 	}
 
-	if (!CurrentPlayer->GetControls().Seek(Time))
+	FMediaSeekParams SeekParams;
+	// v2 timing players are expected to use the new sequence index we set up.
+	int32 SequenceIndexNow = SeekIndex;
+	if (CurrentPlayer->GetPlayerFeatureFlag(IMediaPlayer::EFeatureFlag::UsePlaybackTimingV2))
+	{
+		SeekIndex += (GetUnpausedRate() < 0.0f) ? -1 : 1;
+		SeekParams.NewSequenceIndex = SeekIndex;
+	}
+	else
 	{
+		SeekIndex = 0;
+	}
+	// Issue the seek.
+	if (!CurrentPlayer->GetControls().Seek(Time, SeekParams))
+	{
+		// If that failed restore the sequence index.
+		SeekIndex = SequenceIndexNow;
 		return false;
 	}
 
@@ -1227,8 +1242,7 @@ bool FMediaPlayerFacade::Seek(const FTimespan& InTime)
 	{
 		// Yes. Flush only the facade side of the system as needed for seeks
 		// (the player is expected to flush its internal queues as needed itself)
-		check(CurrentPlayer->GetPlayerFeatureFlag(IMediaPlayer::EFeatureFlag::PlayerUsesInternalFlushOnSeek));
-		Flush(true, true);
+		Flush(CurrentPlayer->GetPlayerFeatureFlag(IMediaPlayer::EFeatureFlag::PlayerUsesInternalFlushOnSeek), true);
 	}
 	else
 	{
@@ -1730,13 +1744,11 @@ void FMediaPlayerFacade::Flush(bool bExcludePlayer, bool bOnSeek)
 	LastAudioRenderedSampleTime.Invalidate();
 	if (bOnSeek)
 	{
-		SeekIndex += (GetUnpausedRate() < 0.0f) ? -1 : 1;
 		BlockOnRange.OnSeek(SeekIndex);
 	}
 	else
 	{
 		BlockOnRange.OnFlush();
-		SeekIndex = 0;
 	}
 
 	if (Player.IsValid() && Player->GetPlayerFeatureFlag(IMediaPlayer::EFeatureFlag::UsePlaybackTimingV2))
