[WmfMedia] Disable AutoRTFM instrumentation.

AutoRTFM instrumentation cannot be enabled with exceptions.

diff --git a/Engine/Plugins/Media/WmfMedia/Source/WmfMedia/WmfMedia.Build.cs b/Engine/Plugins/Media/WmfMedia/Source/WmfMedia/WmfMedia.Build.cs
--- a/Engine/Plugins/Media/WmfMedia/Source/WmfMedia/WmfMedia.Build.cs
+++ b/Engine/Plugins/Media/WmfMedia/Source/WmfMedia/WmfMedia.Build.cs
@@ -7,6 +7,7 @@ namespace UnrealBuildTool.Rules
 		public WmfMedia(ReadOnlyTargetRules Target) : base(Target)
 		{
             bEnableExceptions = true;
+			bDisableAutoRTFMInstrumentation = true; // AutoRTFM cannot be used with exceptions
 
             DynamicallyLoadedModuleNames.AddRange(
 				new string[] {
