WmfMedia: Fix for incorrect dx12 decoding using the uncropped image size resulting in green rows at the bottom of the image

#rnx

diff --git a/Engine/Plugins/Media/WmfMedia/Source/WmfMedia/Private/Wmf/WmfMediaStreamSink.cpp b/Engine/Plugins/Media/WmfMedia/Source/WmfMedia/Private/Wmf/WmfMediaStreamSink.cpp
--- a/Engine/Plugins/Media/WmfMedia/Source/WmfMedia/Private/Wmf/WmfMediaStreamSink.cpp
+++ b/Engine/Plugins/Media/WmfMedia/Source/WmfMedia/Private/Wmf/WmfMediaStreamSink.cpp
@@ -860,6 +860,23 @@ void FWmfMediaStreamSink::CopyTextureAndEnqueueSample(IMFSample* pSample)
 	UINT32 DimY;
 	MFGetAttributeSize(CurrentMediaType, MF_MT_FRAME_SIZE, &DimX, &DimY);
 
+	MFVideoArea videoArea = {};
+	Result = CurrentMediaType->GetBlob(MF_MT_MINIMUM_DISPLAY_APERTURE, (UINT8*)&videoArea, sizeof(MFVideoArea), nullptr);
+	if (FAILED(Result))
+	{
+		Result = CurrentMediaType->GetBlob(MF_MT_GEOMETRIC_APERTURE, (UINT8*)&videoArea, sizeof(MFVideoArea), nullptr);
+		if (FAILED(Result))
+		{
+			videoArea.OffsetX.fract = 0;
+			videoArea.OffsetX.value = 0;
+			videoArea.OffsetY.fract = 0;
+			videoArea.OffsetY.value = 0;
+			videoArea.Area.cx = DimX;
+			videoArea.Area.cy = DimY;
+		}
+	}
+
+
 	LONGLONG SampleTime = 0;
 	LONGLONG SampleDuration = 0;
 
@@ -1056,7 +1073,7 @@ void FWmfMediaStreamSink::CopyTextureAndEnqueueSample(IMFSample* pSample)
 					}
 
 					uint32 Pitch = (DimX / GPixelFormats[PixelFormat].BlockSizeX) * GPixelFormats[PixelFormat].BlockBytes;
-					TextureSample->Initialize(Data, BufferSize, FIntPoint(DimX, SampleDimY), FIntPoint(DimX, DimY), MediaTextureSampleFormat, Pitch, AdjustTimeStamp(FTimespan::FromMicroseconds(SampleTime / 10), EMediaTrackType::Video), FTimespan::FromMicroseconds(SampleDuration / 10));
+					TextureSample->Initialize(Data, BufferSize, FIntPoint(DimX, SampleDimY), FIntPoint(videoArea.Area.cx, videoArea.Area.cy), MediaTextureSampleFormat, Pitch, AdjustTimeStamp(FTimespan::FromMicroseconds(SampleTime / 10), EMediaTrackType::Video), FTimespan::FromMicroseconds(SampleDuration / 10));
 					pBuffer->Unlock();
 
 					TextureSample->SetPixelFormat(PixelFormat);
