WmfMedia: Fix for incorrect dx11 decoding using the uncropped image size resulting in duplicated rows at the bottom of the image

#rnx

diff --git a/Engine/Plugins/Media/WmfMedia/Source/WmfMedia/Private/Wmf/WmfMediaHardwareVideoDecodingRendering.cpp b/Engine/Plugins/Media/WmfMedia/Source/WmfMedia/Private/Wmf/WmfMediaHardwareVideoDecodingRendering.cpp
--- a/Engine/Plugins/Media/WmfMedia/Source/WmfMedia/Private/Wmf/WmfMediaHardwareVideoDecodingRendering.cpp
+++ b/Engine/Plugins/Media/WmfMedia/Source/WmfMedia/Private/Wmf/WmfMediaHardwareVideoDecodingRendering.cpp
@@ -189,7 +189,7 @@ bool FWmfMediaHardwareVideoDecodingParameters::ConvertTextureFormat_RenderThread
 			const FHardwareVideoDecodingShaderParams Parameters = FWmfMediaHardwareVideoDecodingShader::GetParameters(Y_SRV, UV_SRV, InSample->IsOutputSrgb());
 
 			SetShaderParameters(RHICmdList, VertexShader, VertexShader.GetVertexShader(), Parameters);
-			SetShaderParameters(RHICmdList, PixelShader, VertexShader.GetPixelShader(), Parameters);
+			SetShaderParameters(RHICmdList, PixelShader, PixelShader.GetPixelShader(), Parameters);
 		}
 		else if ((InSample->GetFormat() == EMediaTextureSampleFormat::CharBGRA) ||
 			(InSample->GetFormat() == EMediaTextureSampleFormat::DXT1) ||
diff --git a/Engine/Plugins/Media/WmfMedia/Source/WmfMedia/Private/Wmf/WmfMediaStreamSink.cpp b/Engine/Plugins/Media/WmfMedia/Source/WmfMedia/Private/Wmf/WmfMediaStreamSink.cpp
--- a/Engine/Plugins/Media/WmfMedia/Source/WmfMedia/Private/Wmf/WmfMediaStreamSink.cpp
+++ b/Engine/Plugins/Media/WmfMedia/Source/WmfMedia/Private/Wmf/WmfMediaStreamSink.cpp
@@ -1011,7 +1011,7 @@ void FWmfMediaStreamSink::CopyTextureAndEnqueueSample(IMFSample* pSample)
 					Owner->GetDevice(),
 					AdjustTimeStamp(FTimespan::FromMicroseconds(SampleTime / 10), EMediaTrackType::Video),
 					FTimespan::FromMicroseconds(SampleDuration / 10),
-					FIntPoint(DimX, DimY),
+					FIntPoint(videoArea.Area.cx, videoArea.Area.cy),
 					PixelFormat,
 					MediaTextureSampleFormat);
 
@@ -1024,8 +1024,8 @@ void FWmfMediaStreamSink::CopyTextureAndEnqueueSample(IMFSample* pSample)
 				SrcBox.left = 0;
 				SrcBox.top = 0;
 				SrcBox.front = 0;
-				SrcBox.right = DimX;
-				SrcBox.bottom = DimY;
+				SrcBox.right = videoArea.Area.cx;
+				SrcBox.bottom = videoArea.Area.cy;
 				SrcBox.back = 1;
 
 				UE_LOG(LogWmfMedia, VeryVerbose, TEXT("CopySubresourceRegion() ViewIndex:%d Time:%f"), dwViewIndex, FTimespan::FromMicroseconds(SampleTime / 10).GetTotalSeconds());
