WmfMedia: Clamping number of inflight requests in case ProcessSample() is invoked multiple times in succession.

#rnx

diff --git a/Engine/Plugins/Media/WmfMedia/Source/WmfMedia/Private/Wmf/WmfMediaStreamSink.cpp b/Engine/Plugins/Media/WmfMedia/Source/WmfMedia/Private/Wmf/WmfMediaStreamSink.cpp
--- a/Engine/Plugins/Media/WmfMedia/Source/WmfMedia/Private/Wmf/WmfMediaStreamSink.cpp
+++ b/Engine/Plugins/Media/WmfMedia/Source/WmfMedia/Private/Wmf/WmfMediaStreamSink.cpp
@@ -692,10 +692,8 @@ STDMETHODIMP FWmfMediaStreamSink::ProcessSample(__RPC__in_opt IMFSample* pSample
 		return MF_E_SHUTDOWN;
 	}
 
-	--NumRequestsInFlight;
-	check(NumRequestsInFlight >= 0);
-	--NumRequestsInFlightWithSystem;
-	check(NumRequestsInFlightWithSystem >= 0);
+	NumRequestsInFlight = FMath::Max(NumRequestsInFlight-1, 0);
+	NumRequestsInFlightWithSystem = FMath::Max(NumRequestsInFlightWithSystem-1, 0);
 
 	bool bIdleWaitInProgress = !!OnIdleAction;
 
